- Fix horde issue 82516, EOSVoiceChat missing Android as a whitelisted platform

diff --git a/Engine/Plugins/Online/VoiceChat/EOSVoiceChat/EOSVoiceChat.uplugin b/Engine/Plugins/Online/VoiceChat/EOSVoiceChat/EOSVoiceChat.uplugin
--- a/Engine/Plugins/Online/VoiceChat/EOSVoiceChat/EOSVoiceChat.uplugin
+++ b/Engine/Plugins/Online/VoiceChat/EOSVoiceChat/EOSVoiceChat.uplugin
@@ -15,6 +15,7 @@
 			"Name": "EOSVoiceChat",
 			"Type": "ClientOnlyNoCommandlet",
 			"WhitelistPlatforms": [
+				"Android",
 				"Mac",
 				"Win64"
 			]
