- Fix some EOSVoiceChat platform extensions still trying to use RunOnEOS/GameThread which should have been removed

diff --git a/Engine/Plugins/Online/VoiceChat/EOSVoiceChat/Source/EOSVoiceChat/Public/EOSVoiceChat.h b/Engine/Plugins/Online/VoiceChat/EOSVoiceChat/Source/EOSVoiceChat/Public/EOSVoiceChat.h
--- a/Engine/Plugins/Online/VoiceChat/EOSVoiceChat/Source/EOSVoiceChat/Public/EOSVoiceChat.h
+++ b/Engine/Plugins/Online/VoiceChat/EOSVoiceChat/Source/EOSVoiceChat/Public/EOSVoiceChat.h
@@ -148,10 +148,6 @@ public:
 
 	EOS_HRTC GetRtcInterface() const { return InitSession.EosRtcInterface; }
 	EOS_HLobby GetLobbyInterface() const { return InitSession.EosLobbyInterface; }
-	void RunOnEOSThread(const TFunction<void(void)>& Task);
-	void RunOnGameThread(const TFunction<void(void)>& Task);
-	void SynchronizeWithEOSThread();
-	bool IsOnEOSThread() const;
 
 	const TArray<FVoiceChatDeviceInfo>& GetCachedInputDeviceInfos() const { return InitSession.CachedInputDeviceInfos; }
 	const TArray<FVoiceChatDeviceInfo>& GetCachedOutputDeviceInfos() const { return InitSession.CachedOutputDeviceInfos; }
