-Added channel name argument to FOnBeforeRecvAudioRenderedDelegate

diff --git a/Engine/Plugins/Online/VoiceChat/EOSVoiceChat/Source/EOSVoiceChat/Private/EOSVoiceChatUser.cpp b/Engine/Plugins/Online/VoiceChat/EOSVoiceChat/Source/EOSVoiceChat/Private/EOSVoiceChatUser.cpp
--- a/Engine/Plugins/Online/VoiceChat/EOSVoiceChat/Source/EOSVoiceChat/Private/EOSVoiceChatUser.cpp
+++ b/Engine/Plugins/Online/VoiceChat/EOSVoiceChat/Source/EOSVoiceChat/Private/EOSVoiceChatUser.cpp
@@ -2205,14 +2205,6 @@ void EOS_CALL FEOSVoiceChatUser::OnChannelAudioBeforeRenderStatic(const EOS_RTCA
 	}
 }
 
-static FString HACK_ProductUserIdToString(EOS_ProductUserId UserId)
-{
-	char UserIdStr[EOS_PRODUCTUSERID_MAX_LENGTH];
-	int32_t Length = sizeof(UserIdStr);
-	EOS_ProductUserId_ToString(UserId, UserIdStr, &Length);
-	return UTF8_TO_TCHAR(UserIdStr);
-}
-
 void FEOSVoiceChatUser::OnChannelAudioBeforeRender(const EOS_RTCAudio_AudioBeforeRenderCallbackInfo* CallbackInfo)
 {
 	if (EOS_RTCAudio_AudioBuffer* Buffer = CallbackInfo->Buffer)
@@ -2223,10 +2215,11 @@ void FEOSVoiceChatUser::OnChannelAudioBeforeRender(const EOS_RTCAudio_AudioBefor
 
 			// TODO EOS doesn't tell us if it's silence or not, maybe need to compare all the samples to some threshold?
 			const bool bIsSilence = false;
-			const FString PlayerName = HACK_ProductUserIdToString(CallbackInfo->ParticipantId);
+			const FString PlayerName = LexToString(CallbackInfo->ParticipantId);
+			const FString ChannelName = UTF8_TO_TCHAR(CallbackInfo->RoomName);
 
 			FScopeLock Lock(&BeforeRecvAudioRenderedLock);
-			OnVoiceChatBeforeRecvAudioRenderedDelegate.Broadcast(Samples, Buffer->SampleRate, Buffer->Channels, bIsSilence, PlayerName);
+			OnVoiceChatBeforeRecvAudioRenderedDelegate.Broadcast(Samples, Buffer->SampleRate, Buffer->Channels, bIsSilence, ChannelName, PlayerName);
 		}
 		else
 		{
diff --git a/Engine/Plugins/Online/VoiceChat/VoiceChat/Source/Public/VoiceChat.h b/Engine/Plugins/Online/VoiceChat/VoiceChat/Source/Public/VoiceChat.h
--- a/Engine/Plugins/Online/VoiceChat/VoiceChat/Source/Public/VoiceChat.h
+++ b/Engine/Plugins/Online/VoiceChat/VoiceChat/Source/Public/VoiceChat.h
@@ -113,7 +113,7 @@ DECLARE_MULTICAST_DELEGATE_ThreeParams(FOnVoiceChatPlayerVolumeUpdatedDelegate,
 DECLARE_MULTICAST_DELEGATE_ThreeParams(FOnVoiceChatRecordSamplesAvailableDelegate, TArrayView<const int16> /* PcmSamples */, int /* SampleRate */, int /* Channels */);
 DECLARE_MULTICAST_DELEGATE_ThreeParams(FOnVoiceChatAfterCaptureAudioReadDelegate, TArrayView<int16> /* PcmSamples */, int /* SampleRate */, int /* Channels */);
 DECLARE_MULTICAST_DELEGATE_FourParams(FOnVoiceChatBeforeCaptureAudioSentDelegate, TArrayView<const int16> /* PcmSamples */, int /* SampleRate */, int /* Channels */, bool /* bIsSpeaking */);
-DECLARE_MULTICAST_DELEGATE_FiveParams(FOnVoiceChatBeforeRecvAudioRenderedDelegate, TArrayView<int16> /* PcmSamples */, int /* SampleRate */, int /* Channels */, bool /* bIsSilence */, const FString& /* PlayerName */);
+DECLARE_MULTICAST_DELEGATE_SixParams(FOnVoiceChatBeforeRecvAudioRenderedDelegate, TArrayView<int16> /* PcmSamples */, int /* SampleRate */, int /* Channels */, bool /* bIsSilence */, const FString& /*ChannelName*/, const FString& /* PlayerName */);
 
 class IVoiceChatUser
 {
