- Remove bForce bools from EOSVoiceChat Apply... functions, and always call the SDK functions, to remove opportunity for mismatch bugs. Removed some now redundant caching variables.
- Document IVoiceChat::SetChannelPlayerMuted lifetime.

diff --git a/Engine/Plugins/Online/VoiceChat/EOSVoiceChat/Source/EOSVoiceChat/Private/EOSVoiceChatUser.cpp b/Engine/Plugins/Online/VoiceChat/EOSVoiceChat/Source/EOSVoiceChat/Private/EOSVoiceChatUser.cpp
--- a/Engine/Plugins/Online/VoiceChat/EOSVoiceChat/Source/EOSVoiceChat/Private/EOSVoiceChatUser.cpp
+++ b/Engine/Plugins/Online/VoiceChat/EOSVoiceChat/Source/EOSVoiceChat/Private/EOSVoiceChatUser.cpp
@@ -120,7 +120,7 @@ void FEOSVoiceChatUser::SetSetting(const FString& Name, const FString& Value)
 			if (bIsNotListening != ChannelSession.bIsNotListening)
 			{
 				ChannelSession.bIsNotListening = bIsNotListening;
-				ApplySendingOptions(ChannelSession, true);
+				ApplySendingOptions(ChannelSession);
 				ApplyReceivingOptions(ChannelSession);
 				for (TPair<FString, FChannelParticipant>& ChannelParticipantPair : ChannelSession.Participants)
 				{
@@ -128,7 +128,7 @@ void FEOSVoiceChatUser::SetSetting(const FString& Name, const FString& Value)
 					FChannelParticipant& ChannelParticipant = ChannelParticipantPair.Value;
 					if (!ChannelSession.IsLocalUser(ChannelParticipant))
 					{
-						ApplyPlayerReceivingOptions(GetGlobalParticipant(ChannelParticipantName), ChannelSession, ChannelParticipant, true);
+						ApplyPlayerReceivingOptions(GetGlobalParticipant(ChannelParticipantName), ChannelSession, ChannelParticipant);
 					}
 				}
 			}
@@ -1178,26 +1178,22 @@ void FEOSVoiceChatUser::ApplyAudioOutputOptions()
 
 void FEOSVoiceChatUser::ApplyPlayerBlock(const FGlobalParticipant& GlobalParticipant, const FChannelSession& ChannelSession, FChannelParticipant& ChannelParticipant)
 {
-	if (GlobalParticipant.bBlocked != ChannelParticipant.bBlocked)
-	{
-		EOSVOICECHATUSER_LOG(Log, TEXT("ApplyPlayerBlock ChannelName=[%s] PlayerName=[%s] bBlocked=%s"), *ChannelSession.ChannelName, *GlobalParticipant.PlayerName, *LexToString(GlobalParticipant.bBlocked));
-		ChannelParticipant.bBlocked = GlobalParticipant.bBlocked;
+	EOSVOICECHATUSER_LOG(Verbose, TEXT("ApplyPlayerBlock ChannelName=[%s] PlayerName=[%s] bBlocked=%s"), *ChannelSession.ChannelName, *GlobalParticipant.PlayerName, *LexToString(GlobalParticipant.bBlocked));
 
-		const FTCHARToUTF8 Utf8RoomName(*ChannelSession.ChannelName);
-		const FTCHARToUTF8 Utf8ParticipantId(*GlobalParticipant.PlayerName);
+	const FTCHARToUTF8 Utf8RoomName(*ChannelSession.ChannelName);
+	const FTCHARToUTF8 Utf8ParticipantId(*GlobalParticipant.PlayerName);
 
-		EOS_RTC_BlockParticipantOptions Options = {};
-		Options.ApiVersion = EOS_RTC_BLOCKPARTICIPANT_API_LATEST;
-		static_assert(EOS_RTC_BLOCKPARTICIPANT_API_LATEST == 1, "EOS_RTC_BlockParticipantOptions updated, check new fields");
-		Options.LocalUserId = LoginSession.LocalUserProductUserId;
-		Options.RoomName = Utf8RoomName.Get();
-		Options.ParticipantId = EOS_ProductUserId_FromString(Utf8ParticipantId.Get());
-		Options.bBlocked = ChannelParticipant.bBlocked;
+	EOS_RTC_BlockParticipantOptions Options = {};
+	Options.ApiVersion = EOS_RTC_BLOCKPARTICIPANT_API_LATEST;
+	static_assert(EOS_RTC_BLOCKPARTICIPANT_API_LATEST == 1, "EOS_RTC_BlockParticipantOptions updated, check new fields");
+	Options.LocalUserId = LoginSession.LocalUserProductUserId;
+	Options.RoomName = Utf8RoomName.Get();
+	Options.ParticipantId = EOS_ProductUserId_FromString(Utf8ParticipantId.Get());
+	Options.bBlocked = GlobalParticipant.bBlocked;
 
-		CSV_SCOPED_TIMING_STAT_EXCLUSIVE(EOSVoiceChat);
-		QUICK_SCOPE_CYCLE_COUNTER(EOS_RTC_BlockParticipant);
-		EOS_RTC_BlockParticipant(GetRtcInterface(), &Options, this, &FEOSVoiceChatUser::OnBlockParticipantStatic);
-	}
+	CSV_SCOPED_TIMING_STAT_EXCLUSIVE(EOSVoiceChat);
+	QUICK_SCOPE_CYCLE_COUNTER(EOS_RTC_BlockParticipant);
+	EOS_RTC_BlockParticipant(GetRtcInterface(), &Options, this, &FEOSVoiceChatUser::OnBlockParticipantStatic);
 };
 
 void FEOSVoiceChatUser::ApplyReceivingOptions(const FChannelSession& ChannelSession)
@@ -1217,38 +1213,28 @@ void FEOSVoiceChatUser::ApplyReceivingOptions(const FChannelSession& ChannelSess
 	EOS_RTCAudio_UpdateReceiving(EOS_RTC_GetAudioInterface(GetRtcInterface()), &UpdateReceivingOptions, this, &FEOSVoiceChatUser::OnUpdateReceivingAudioStatic);
 }
 
-void FEOSVoiceChatUser::ApplyPlayerReceivingOptions(const FGlobalParticipant& GlobalParticipant, const FChannelSession& ChannelSession, FChannelParticipant& ChannelParticipant, const bool bForce)
+void FEOSVoiceChatUser::ApplyPlayerReceivingOptions(const FGlobalParticipant& GlobalParticipant, const FChannelSession& ChannelSession, FChannelParticipant& ChannelParticipant)
 {
 	// TODO ChannelParticipant should be const, and we should update the bAudioMuted (the *actual* state, not the *desired* state) in the OnUpdateReceiving callback.
 	check(IsLoggedIn());
 
-	const bool bAudioDisabled = GlobalParticipant.bAudioMuted || ChannelSession.bIsNotListening || ChannelParticipant.bMutedInChannel;
-
-	if (ChannelParticipant.bAudioDisabled != bAudioDisabled || bForce)
-	{
-		ChannelParticipant.bAudioDisabled = bAudioDisabled;
-
-		EOSVOICECHATUSER_LOG(Log, TEXT("ApplyPlayerReceivingOptions ChannelName=[%s] PlayerName=[%s] bAudioMuted=[%s]"), *ChannelSession.ChannelName, *GlobalParticipant.PlayerName, *LexToString(GlobalParticipant.bAudioMuted));
+	ChannelParticipant.bAudioDisabled = GlobalParticipant.bAudioMuted || ChannelSession.bIsNotListening || ChannelParticipant.bMutedInChannel;
+	EOSVOICECHATUSER_LOG(Verbose, TEXT("ApplyPlayerReceivingOptions ChannelName=[%s] PlayerName=[%s] bAudioMuted=[%s]"), *ChannelSession.ChannelName, *GlobalParticipant.PlayerName, *LexToString(GlobalParticipant.bAudioMuted));
 
-		const FTCHARToUTF8 Utf8RoomName(*ChannelSession.ChannelName);
-		const FTCHARToUTF8 Utf8ParticipantId(*GlobalParticipant.PlayerName);
+	const FTCHARToUTF8 Utf8RoomName(*ChannelSession.ChannelName);
+	const FTCHARToUTF8 Utf8ParticipantId(*GlobalParticipant.PlayerName);
 
-		EOS_RTCAudio_UpdateReceivingOptions UpdateReceivingOptions = {};
-		UpdateReceivingOptions.ApiVersion = EOS_RTCAUDIO_UPDATERECEIVING_API_LATEST;
-		static_assert(EOS_RTCAUDIO_UPDATERECEIVING_API_LATEST == 1, "EOS_RTCAudio_UpdateReceivingOptions updated, check new fields");
-		UpdateReceivingOptions.LocalUserId = LoginSession.LocalUserProductUserId;
-		UpdateReceivingOptions.RoomName = Utf8RoomName.Get();
-		UpdateReceivingOptions.ParticipantId = EOS_ProductUserId_FromString(Utf8ParticipantId.Get());
-		UpdateReceivingOptions.bAudioEnabled = ChannelParticipant.bAudioDisabled ? EOS_FALSE : EOS_TRUE;
+	EOS_RTCAudio_UpdateReceivingOptions UpdateReceivingOptions = {};
+	UpdateReceivingOptions.ApiVersion = EOS_RTCAUDIO_UPDATERECEIVING_API_LATEST;
+	static_assert(EOS_RTCAUDIO_UPDATERECEIVING_API_LATEST == 1, "EOS_RTCAudio_UpdateReceivingOptions updated, check new fields");
+	UpdateReceivingOptions.LocalUserId = LoginSession.LocalUserProductUserId;
+	UpdateReceivingOptions.RoomName = Utf8RoomName.Get();
+	UpdateReceivingOptions.ParticipantId = EOS_ProductUserId_FromString(Utf8ParticipantId.Get());
+	UpdateReceivingOptions.bAudioEnabled = ChannelParticipant.bAudioDisabled ? EOS_FALSE : EOS_TRUE;
 
-		CSV_SCOPED_TIMING_STAT_EXCLUSIVE(EOSVoiceChat);
-		QUICK_SCOPE_CYCLE_COUNTER(EOS_RTCAudio_UpdateReceiving);
-		EOS_RTCAudio_UpdateReceiving(EOS_RTC_GetAudioInterface(GetRtcInterface()), &UpdateReceivingOptions, this, &FEOSVoiceChatUser::OnUpdateReceivingAudioStatic);
-	}
-	else
-	{
-		EOSVOICECHATUSER_LOG(VeryVerbose, TEXT("ApplyPlayerReceivingOptions ChannelName=[%s] PlayerName=[%s] mute state up to date"), *ChannelSession.ChannelName, *GlobalParticipant.PlayerName);
-	}
+	CSV_SCOPED_TIMING_STAT_EXCLUSIVE(EOSVoiceChat);
+	QUICK_SCOPE_CYCLE_COUNTER(EOS_RTCAudio_UpdateReceiving);
+	EOS_RTCAudio_UpdateReceiving(EOS_RTC_GetAudioInterface(GetRtcInterface()), &UpdateReceivingOptions, this, &FEOSVoiceChatUser::OnUpdateReceivingAudioStatic);
 }
 
 void FEOSVoiceChatUser::ApplySendingOptions()
@@ -1263,34 +1249,29 @@ void FEOSVoiceChatUser::ApplySendingOptions()
 	}
 }
 
-void FEOSVoiceChatUser::ApplySendingOptions(FChannelSession& ChannelSession, const bool bForce)
+void FEOSVoiceChatUser::ApplySendingOptions(FChannelSession& ChannelSession)
 {
 	const bool bCanTransmitToChannel = TransmitState.Mode == EVoiceChatTransmitMode::All ||
 		(TransmitState.Mode == EVoiceChatTransmitMode::Channel && TransmitState.ChannelName == ChannelSession.ChannelName);
-	const bool bAudioEnabled = bCanTransmitToChannel && !AudioInputOptions.bMuted && !ChannelSession.bIsNotListening;
 
-	if (ChannelSession.DesiredSendingState.bAudioEnabled != bAudioEnabled
-		|| bForce)
-	{
-		ChannelSession.DesiredSendingState.bAudioEnabled = bAudioEnabled;
+	ChannelSession.DesiredSendingState.bAudioEnabled = bCanTransmitToChannel && !AudioInputOptions.bMuted && !ChannelSession.bIsNotListening;
 
-		const FTCHARToUTF8 Utf8RoomName(*ChannelSession.ChannelName);
+	const FTCHARToUTF8 Utf8RoomName(*ChannelSession.ChannelName);
 		
-		EOS_RTCAudio_UpdateSendingOptions UpdateSendingOptions = {};
-		UpdateSendingOptions.ApiVersion = EOS_RTCAUDIO_UPDATESENDING_API_LATEST;
-		static_assert(EOS_RTCAUDIO_UPDATESENDING_API_LATEST == 1, "EOS_RTCAudio_UpdateSendingOptions updated, check new fields");
-		UpdateSendingOptions.LocalUserId = LoginSession.LocalUserProductUserId;
-		UpdateSendingOptions.RoomName = Utf8RoomName.Get();
-		UpdateSendingOptions.AudioStatus = ChannelSession.bIsNotListening
-			? EOS_ERTCAudioStatus::EOS_RTCAS_NotListeningDisabled
-			: ChannelSession.DesiredSendingState.bAudioEnabled
-			? EOS_ERTCAudioStatus::EOS_RTCAS_Enabled
-			: EOS_ERTCAudioStatus::EOS_RTCAS_Disabled;
+	EOS_RTCAudio_UpdateSendingOptions UpdateSendingOptions = {};
+	UpdateSendingOptions.ApiVersion = EOS_RTCAUDIO_UPDATESENDING_API_LATEST;
+	static_assert(EOS_RTCAUDIO_UPDATESENDING_API_LATEST == 1, "EOS_RTCAudio_UpdateSendingOptions updated, check new fields");
+	UpdateSendingOptions.LocalUserId = LoginSession.LocalUserProductUserId;
+	UpdateSendingOptions.RoomName = Utf8RoomName.Get();
+	UpdateSendingOptions.AudioStatus = ChannelSession.bIsNotListening
+		? EOS_ERTCAudioStatus::EOS_RTCAS_NotListeningDisabled
+		: ChannelSession.DesiredSendingState.bAudioEnabled
+		? EOS_ERTCAudioStatus::EOS_RTCAS_Enabled
+		: EOS_ERTCAudioStatus::EOS_RTCAS_Disabled;
 
-		CSV_SCOPED_TIMING_STAT_EXCLUSIVE(EOSVoiceChat);
-		QUICK_SCOPE_CYCLE_COUNTER(EOS_RTCAudio_UpdateSending);
-		EOS_RTCAudio_UpdateSending(EOS_RTC_GetAudioInterface(GetRtcInterface()), &UpdateSendingOptions, this, &FEOSVoiceChatUser::OnUpdateSendingAudioStatic);
-	}
+	CSV_SCOPED_TIMING_STAT_EXCLUSIVE(EOSVoiceChat);
+	QUICK_SCOPE_CYCLE_COUNTER(EOS_RTCAudio_UpdateSending);
+	EOS_RTCAudio_UpdateSending(EOS_RTC_GetAudioInterface(GetRtcInterface()), &UpdateSendingOptions, this, &FEOSVoiceChatUser::OnUpdateSendingAudioStatic);
 }
 
 void FEOSVoiceChatUser::BindChannelCallbacks(FChannelSession& ChannelSession)
@@ -1636,7 +1617,7 @@ void FEOSVoiceChatUser::OnJoinRoom(const EOS_RTC_JoinRoomCallbackInfo* CallbackI
 				ChannelParticipant.PlayerName = ChannelSession->PlayerName;
 				OnVoiceChatPlayerAddedDelegate.Broadcast(ChannelSession->ChannelName, ChannelSession->PlayerName);
 
-				ApplySendingOptions(*ChannelSession, true);
+				ApplySendingOptions(*ChannelSession);
 			}
 			else
 			{
diff --git a/Engine/Plugins/Online/VoiceChat/EOSVoiceChat/Source/EOSVoiceChat/Public/EOSVoiceChatUser.h b/Engine/Plugins/Online/VoiceChat/EOSVoiceChat/Source/EOSVoiceChat/Public/EOSVoiceChatUser.h
--- a/Engine/Plugins/Online/VoiceChat/EOSVoiceChat/Source/EOSVoiceChat/Public/EOSVoiceChatUser.h
+++ b/Engine/Plugins/Online/VoiceChat/EOSVoiceChat/Source/EOSVoiceChat/Public/EOSVoiceChatUser.h
@@ -129,8 +129,6 @@ protected:
 		bool bTalking = false;
 		// Combined audio mute and isListening state
 		bool bAudioDisabled = false;
-		// Current block state
-		bool bBlocked = false;
 		// Desired channel mute state
 		bool bMutedInChannel = false;
 		// Current audio status
@@ -308,9 +306,9 @@ protected:
 	void ApplyAudioOutputOptions();
 	void ApplyPlayerBlock(const FGlobalParticipant& GlobalParticipant, const FChannelSession& ChannelSession, FChannelParticipant& ChannelParticipant);
 	void ApplyReceivingOptions(const FChannelSession& ChannelSession);
-	void ApplyPlayerReceivingOptions(const FGlobalParticipant& GlobalParticipant, const FChannelSession& ChannelSession, FChannelParticipant& ChannelParticipant, bool bForce = false);
+	void ApplyPlayerReceivingOptions(const FGlobalParticipant& GlobalParticipant, const FChannelSession& ChannelSession, FChannelParticipant& ChannelParticipant);
 	void ApplySendingOptions();
-	void ApplySendingOptions(FChannelSession& ChannelSession, const bool bForce = false);
+	void ApplySendingOptions(FChannelSession& ChannelSession);
 	void BindChannelCallbacks(FChannelSession& ChannelSession);
 	void UnbindChannelCallbacks(FChannelSession& ChannelSession);
 	void LeaveChannelInternal(const FString& ChannelName, const FOnVoiceChatChannelLeaveCompleteDelegate& Delegate);
diff --git a/Engine/Plugins/Online/VoiceChat/VoiceChat/Source/Public/VoiceChat.h b/Engine/Plugins/Online/VoiceChat/VoiceChat/Source/Public/VoiceChat.h
--- a/Engine/Plugins/Online/VoiceChat/VoiceChat/Source/Public/VoiceChat.h
+++ b/Engine/Plugins/Online/VoiceChat/VoiceChat/Source/Public/VoiceChat.h
@@ -525,7 +525,9 @@ public:
 	virtual bool IsPlayerMuted(const FString& PlayerName) const = 0;
 
 	/**
-	 * Mute or unmute a player in a given channel
+	 * Mute or unmute a player in a given channel.
+	 * Only has effect if the player is currently in the channel, and only until you/they leave the channel.
+	 * i.e. this is expected to be [re]called by OnVoiceChatPlayerAdded delegates.
 	 *
 	 * @param ChannelName Channel in which to mute the player
 	 * @param PlayerName Player to mute in channel
