- Add missing API_VERSION static asserts and fix some mislabelled ones

diff --git a/Engine/Plugins/Online/VoiceChat/EOSVoiceChat/Source/EOSVoiceChat/Private/EOSVoiceChat.cpp b/Engine/Plugins/Online/VoiceChat/EOSVoiceChat/Source/EOSVoiceChat/Private/EOSVoiceChat.cpp
--- a/Engine/Plugins/Online/VoiceChat/EOSVoiceChat/Source/EOSVoiceChat/Private/EOSVoiceChat.cpp
+++ b/Engine/Plugins/Online/VoiceChat/EOSVoiceChat/Source/EOSVoiceChat/Private/EOSVoiceChat.cpp
@@ -956,6 +956,7 @@ TArray<FVoiceChatDeviceInfo> FEOSVoiceChat::GetRtcInputDeviceInfos(int32& OutDef
 	{
 		EOS_RTCAudio_GetAudioInputDeviceByIndexOptions GetByIndexOptions = {};
 		GetByIndexOptions.ApiVersion = EOS_RTCAUDIO_GETAUDIOINPUTDEVICEBYINDEX_API_LATEST;
+		static_assert(EOS_RTCAUDIO_GETAUDIOINPUTDEVICEBYINDEX_API_LATEST == 1, "EOS_RTCAudio_GetAudioInputDeviceByIndexOptions updated, check new fields");
 		GetByIndexOptions.DeviceInfoIndex = Index;
 		if (const EOS_RTCAudio_AudioInputDeviceInfo* DeviceInfo = EOS_RTCAudio_GetAudioInputDeviceByIndex(RTCAudioHandle, &GetByIndexOptions))
 		{
@@ -1001,6 +1002,7 @@ TArray<FVoiceChatDeviceInfo> FEOSVoiceChat::GetRtcOutputDeviceInfos(int32& OutDe
 	{
 		EOS_RTCAudio_GetAudioOutputDeviceByIndexOptions GetByIndexOptions = {};
 		GetByIndexOptions.ApiVersion = EOS_RTCAUDIO_GETAUDIOOUTPUTDEVICEBYINDEX_API_LATEST;
+		static_assert(EOS_RTCAUDIO_GETAUDIOOUTPUTDEVICEBYINDEX_API_LATEST == 1, "EOS_RTCAudio_GetAudioOutputDeviceByIndexOptions updated, check new fields");
 		GetByIndexOptions.DeviceInfoIndex = Index;
 		if (const EOS_RTCAudio_AudioOutputDeviceInfo* DeviceInfo = EOS_RTCAudio_GetAudioOutputDeviceByIndex(RTCAudioHandle, &GetByIndexOptions))
 		{
diff --git a/Engine/Plugins/Online/VoiceChat/EOSVoiceChat/Source/EOSVoiceChat/Private/EOSVoiceChatUser.cpp b/Engine/Plugins/Online/VoiceChat/EOSVoiceChat/Source/EOSVoiceChat/Private/EOSVoiceChatUser.cpp
--- a/Engine/Plugins/Online/VoiceChat/EOSVoiceChat/Source/EOSVoiceChat/Private/EOSVoiceChatUser.cpp
+++ b/Engine/Plugins/Online/VoiceChat/EOSVoiceChat/Source/EOSVoiceChat/Private/EOSVoiceChatUser.cpp
@@ -918,7 +918,7 @@ bool FEOSVoiceChatUser::AddLobbyRoom(const FString& LobbyId)
 
 	EOS_Lobby_GetRTCRoomNameOptions GetRoomNameOptions = {};
 	GetRoomNameOptions.ApiVersion = EOS_LOBBY_GETRTCROOMNAME_API_LATEST;
-	static_assert(EOS_LOBBY_GETRTCROOMNAME_API_LATEST == 1, "EOS_Lobby_GetLobbyRTCRoomNameOptions updated, check new fields");
+	static_assert(EOS_LOBBY_GETRTCROOMNAME_API_LATEST == 1, "EOS_Lobby_GetRTCRoomNameOptions updated, check new fields");
 	GetRoomNameOptions.LobbyId = Utf8LobbyId.Get();
 	GetRoomNameOptions.LocalUserId = LoginSession.LocalUserProductUserId;
 	char Utf8RoomName[256];
@@ -1315,7 +1315,7 @@ void FEOSVoiceChatUser::BindChannelCallbacks(FChannelSession& ChannelSession)
 	{
 		EOS_RTC_AddNotifyParticipantStatusChangedOptions ParticipantStatusChangedOptions = {};
 		ParticipantStatusChangedOptions.ApiVersion = EOS_RTC_ADDNOTIFYPARTICIPANTSTATUSCHANGED_API_LATEST;
-		static_assert(EOS_RTC_ADDNOTIFYPARTICIPANTSTATUSCHANGED_API_LATEST == 1, "EOS_RTC_AddNotifyParticipantJoinedOptions updated, check new fields");
+		static_assert(EOS_RTC_ADDNOTIFYPARTICIPANTSTATUSCHANGED_API_LATEST == 1, "EOS_RTC_AddNotifyParticipantStatusChangedOptions updated, check new fields");
 		ParticipantStatusChangedOptions.LocalUserId = LoginSession.LocalUserProductUserId;
 		ParticipantStatusChangedOptions.RoomName = Utf8RoomName.Get();
 		ChannelSession.OnParticipantStatusChangedNotificationId = EOS_RTC_AddNotifyParticipantStatusChanged(GetRtcInterface(), &ParticipantStatusChangedOptions, this, &FEOSVoiceChatUser::OnChannelParticipantStatusChangedStatic);
