- Provide getter for the platform handle from EOSVoiceChat
- Remove unused getters for cached device info

#9118

diff --git a/Engine/Plugins/Online/VoiceChat/EOSVoiceChat/Source/EOSVoiceChat/Public/EOSVoiceChat.h b/Engine/Plugins/Online/VoiceChat/EOSVoiceChat/Source/EOSVoiceChat/Public/EOSVoiceChat.h
--- a/Engine/Plugins/Online/VoiceChat/EOSVoiceChat/Source/EOSVoiceChat/Public/EOSVoiceChat.h
+++ b/Engine/Plugins/Online/VoiceChat/EOSVoiceChat/Source/EOSVoiceChat/Public/EOSVoiceChat.h
@@ -148,12 +148,10 @@ public:
 	virtual FString InsecureGetJoinToken(const FString& ChannelName, EVoiceChatChannelType ChannelType, TOptional<FVoiceChatChannel3dProperties> Channel3dProperties = TOptional<FVoiceChatChannel3dProperties>()) override;
 	// ~End IVoiceChatUser Interface
 
+	IEOSPlatformHandlePtr GetPlatformHandle() const { return InitSession.EosPlatformHandle; }
 	EOS_HRTC GetRtcInterface() const { return InitSession.EosRtcInterface; }
 	EOS_HLobby GetLobbyInterface() const { return InitSession.EosLobbyInterface; }
 
-	const TArray<FVoiceChatDeviceInfo>& GetCachedInputDeviceInfos() const { return InitSession.CachedInputDeviceInfos; }
-	const TArray<FVoiceChatDeviceInfo>& GetCachedOutputDeviceInfos() const { return InitSession.CachedOutputDeviceInfos; }
-
 protected:
 	virtual void PostInitialize() {};
 	virtual void PreUninitialize() {};
