- Fix failed LeaveRoom calls during shutdown resulting in channel callbacks not being unbound.
- Fix/replace CreateWeakThis to use AsWeak instead of AsShared. This was causing strong reference count to be bumped in audio thread callbacks, causing the object to be destroyed on audio threads sometimes.

diff --git a/Engine/Plugins/Online/VoiceChat/EOSVoiceChat/Source/EOSVoiceChat/Private/EOSVoiceChat.cpp b/Engine/Plugins/Online/VoiceChat/EOSVoiceChat/Source/EOSVoiceChat/Private/EOSVoiceChat.cpp
--- a/Engine/Plugins/Online/VoiceChat/EOSVoiceChat/Source/EOSVoiceChat/Private/EOSVoiceChat.cpp
+++ b/Engine/Plugins/Online/VoiceChat/EOSVoiceChat/Source/EOSVoiceChat/Private/EOSVoiceChat.cpp
@@ -409,7 +409,7 @@ void FEOSVoiceChat::ReleaseUser(IVoiceChatUser* User)
 			&& User->IsLoggedIn())
 		{
 			UE_LOG(LogEOSVoiceChat, Log, TEXT("ReleaseUser User=[%p] Logging out"), User);
-			User->Logout(FOnVoiceChatLogoutCompleteDelegate::CreateLambda([this, WeakThis = CreateWeakThis(), User](const FString& PlayerName, const FVoiceChatResult& Result)
+			User->Logout(FOnVoiceChatLogoutCompleteDelegate::CreateLambda([this, WeakThis = AsWeak(), User](const FString& PlayerName, const FVoiceChatResult& Result)
 			{
 				CHECKPIN();
 
@@ -1216,11 +1216,6 @@ IEOSPlatformHandlePtr FEOSVoiceChat::EOSPlatformCreate(EOS_Platform_Options& Pla
 	return SDKManager.CreatePlatform(PlatformOptions);
 }
 
-FEOSVoiceChatWeakPtr FEOSVoiceChat::CreateWeakThis()
-{
-	return FEOSVoiceChatWeakPtr(AsShared());
-}
-
 const TCHAR* LexToString(FEOSVoiceChat::EConnectionState State)
 {
 	switch (State)
diff --git a/Engine/Plugins/Online/VoiceChat/EOSVoiceChat/Source/EOSVoiceChat/Private/EOSVoiceChatUser.cpp b/Engine/Plugins/Online/VoiceChat/EOSVoiceChat/Source/EOSVoiceChat/Private/EOSVoiceChatUser.cpp
--- a/Engine/Plugins/Online/VoiceChat/EOSVoiceChat/Source/EOSVoiceChat/Private/EOSVoiceChatUser.cpp
+++ b/Engine/Plugins/Online/VoiceChat/EOSVoiceChat/Source/EOSVoiceChat/Private/EOSVoiceChatUser.cpp
@@ -98,6 +98,7 @@ FEOSVoiceChatUser::FEOSVoiceChatUser(FEOSVoiceChat& InEOSVoiceChat)
 
 FEOSVoiceChatUser::~FEOSVoiceChatUser()
 {
+	ClearLoginSession();
 }
 
 #pragma region IVoiceChatUser
@@ -1004,7 +1005,7 @@ bool FEOSVoiceChatUser::RemoveLobbyRoom(const FString& LobbyId)
 		if (ensure(ChannelSession))
 		{
 			ChannelSession->JoinState = EChannelJoinState::NotJoined;
-			UnbindChannelCallbacks(*ChannelSession);
+
 			OnVoiceChatChannelExitedDelegate.Broadcast(ChannelSession->ChannelName, FVoiceChatResult::CreateSuccess());
 
 			LoginSession.LobbyIdToChannelName.Remove(ChannelSession->LobbyId);
@@ -1054,11 +1055,6 @@ void FEOSVoiceChatUser::SetHardwareAECEnabled(bool bEnabled)
 	ApplyAudioInputOptions();
 }
 
-FEOSVoiceChatUserWeakPtr FEOSVoiceChatUser::CreateWeakThis()
-{
-	return FEOSVoiceChatUserWeakPtr(AsShared());
-}
-
 bool FEOSVoiceChatUser::IsInitialized()
 {
 	return EOSVoiceChat.IsInitialized();
@@ -1122,7 +1118,11 @@ const FEOSVoiceChatUser::FChannelSession& FEOSVoiceChatUser::GetChannelSession(c
 
 void FEOSVoiceChatUser::RemoveChannelSession(const FString& ChannelName)
 {
-	LoginSession.ChannelSessions.Remove(ChannelName);
+	if (FChannelSession* ChannelSession = LoginSession.ChannelSessions.Find(ChannelName))
+	{
+		UnbindChannelCallbacks(*ChannelSession);
+		LoginSession.ChannelSessions.Remove(ChannelName);
+	}
 }
 
 void FEOSVoiceChatUser::ApplyAudioInputOptions()
@@ -1565,6 +1565,10 @@ void FEOSVoiceChatUser::LogoutInternal(const FOnVoiceChatLogoutCompleteDelegate&
 
 void FEOSVoiceChatUser::ClearLoginSession()
 {
+	for (TPair<FString, FChannelSession> Pair : LoginSession.ChannelSessions)
+	{
+		RemoveChannelSession(Pair.Key);
+	}
 	LoginSession = FLoginSession();
 }
 
@@ -1623,8 +1627,6 @@ void FEOSVoiceChatUser::OnJoinRoom(const EOS_RTC_JoinRoomCallbackInfo* CallbackI
 			{
 				EOSVOICECHATUSER_LOG(Warning, TEXT("OnJoinRoom ChannelName=[%s] Result=[%s]"), *ChannelName, *LexToString(Result));
 
-				UnbindChannelCallbacks(*ChannelSession);
-
 				ChannelSession->JoinState = EChannelJoinState::NotJoined;
 
 				TriggerCompletionDelegate(ChannelSession->JoinDelegate, ChannelSession->ChannelName, Result);
@@ -1690,7 +1692,6 @@ void FEOSVoiceChatUser::OnLeaveRoom(const EOS_RTC_LeaveRoomCallbackInfo* Callbac
 		if (Result.IsSuccess())
 		{
 			ChannelSession->JoinState = EChannelJoinState::NotJoined;
-			UnbindChannelCallbacks(*ChannelSession);
 		}
 		else
 		{
@@ -1907,7 +1908,6 @@ void FEOSVoiceChatUser::OnChannelDisconnected(const EOS_RTC_DisconnectedCallback
 		else
 		{
 			ChannelSession->JoinState = EChannelJoinState::NotJoined;
-			UnbindChannelCallbacks(*ChannelSession);
 					
 			OnVoiceChatChannelExitedDelegate.Broadcast(ChannelName, Result);
 
diff --git a/Engine/Plugins/Online/VoiceChat/EOSVoiceChat/Source/EOSVoiceChat/Public/EOSVoiceChat.h b/Engine/Plugins/Online/VoiceChat/EOSVoiceChat/Source/EOSVoiceChat/Public/EOSVoiceChat.h
--- a/Engine/Plugins/Online/VoiceChat/EOSVoiceChat/Source/EOSVoiceChat/Public/EOSVoiceChat.h
+++ b/Engine/Plugins/Online/VoiceChat/EOSVoiceChat/Source/EOSVoiceChat/Public/EOSVoiceChat.h
@@ -50,10 +50,8 @@ class IEOSSDKManager;
 using IEOSPlatformHandlePtr = TSharedPtr<class IEOSPlatformHandle, ESPMode::ThreadSafe>;
 
 typedef TSharedPtr<class FEOSVoiceChat, ESPMode::ThreadSafe> FEOSVoiceChatPtr;
-typedef TWeakPtr<class FEOSVoiceChat, ESPMode::ThreadSafe> FEOSVoiceChatWeakPtr;
 typedef TSharedPtr<class FEOSVoiceChatUser, ESPMode::ThreadSafe> FEOSVoiceChatUserPtr;
 typedef TSharedRef<class FEOSVoiceChatUser, ESPMode::ThreadSafe> FEOSVoiceChatUserRef;
-typedef TWeakPtr<class FEOSVoiceChatUser, ESPMode::ThreadSafe> FEOSVoiceChatUserWeakPtr;
 
 class FEOSVoiceChat : public TSharedFromThis<FEOSVoiceChat, ESPMode::ThreadSafe>, public IVoiceChat
 {
@@ -230,8 +228,6 @@ protected:
 	static int64 StaticInstanceIdCount;
 	int64 InstanceId = StaticInstanceIdCount++;
 
-	FEOSVoiceChatWeakPtr CreateWeakThis();
-
 	friend const TCHAR* LexToString(FEOSVoiceChat::EConnectionState State);
 };
 
diff --git a/Engine/Plugins/Online/VoiceChat/EOSVoiceChat/Source/EOSVoiceChat/Public/EOSVoiceChatUser.h b/Engine/Plugins/Online/VoiceChat/EOSVoiceChat/Source/EOSVoiceChat/Public/EOSVoiceChatUser.h
--- a/Engine/Plugins/Online/VoiceChat/EOSVoiceChat/Source/EOSVoiceChat/Public/EOSVoiceChatUser.h
+++ b/Engine/Plugins/Online/VoiceChat/EOSVoiceChat/Source/EOSVoiceChat/Public/EOSVoiceChatUser.h
@@ -317,7 +317,6 @@ protected:
 	void RtcRegisterUser(const FString& UserId);
 	void RtcUnregisterUser(const FString& UserId);
 	void SetHardwareAECEnabled(bool bEnabled);
-	FEOSVoiceChatUserWeakPtr CreateWeakThis();
     
 	// EOS operation callbacks
 	static void EOS_CALL OnJoinRoomStatic(const EOS_RTC_JoinRoomCallbackInfo* CallbackInfo);
