Group cached rtc device manipulations into a separate class

0) Added `FEOSAudioDevicePool` which will be responsible for cached device id operations - enumeration(refreshing), caching a result.
1) Moved GetRtcInputDeviceInfos, GetRtcOutputDeviceInfos, CachedInput|OutputDeviceInfos and DefaultInfos into FEOSAudioDevicePool.
2) Add a common RefreshAudio and per direction calls to run cache updating. Direction-specific has been added due to Switch impl.
3) RefreshAudio uses CompletionDelegates to be able to mask async behavior inside.
4) `FEOSAudioDevicePool` holds as Unique due to its & storying.
5) The common RefreshAudio calls both direction-specific using a shared `FRefreshAudioDevicesState`.
6) Renamed AudioDeviceChanged logs to RefreshAudioDevices (as a related logic was also moved).
7) Updated General `OnAudioDeviceChanged` and Switch `OnAudioDeviceChanged to call RefreshAudioDevices with respective completion delegates.
8) Switch uses a bit different `CHECKPIN()` due to being a derived class.
9) The EOSVoiceChat log tag was extracted to a separate file.

Test plan:
1) Run Incremental https://horde.devtools.epicgames.com/job/62961d55be0ae0b33a3efcd8.
2) Run Editor Win64. Emulate OnAudioDevicesChanged. Refresh is being invoked and cached devices are updated.

diff --git a/Engine/Plugins/Online/VoiceChat/EOSVoiceChat/Source/EOSVoiceChat/Private/EOSAudioDevicePool.cpp b/Engine/Plugins/Online/VoiceChat/EOSVoiceChat/Source/EOSVoiceChat/Private/EOSAudioDevicePool.cpp
new file mode 100644
--- /dev/null
+++ b/Engine/Plugins/Online/VoiceChat/EOSVoiceChat/Source/EOSVoiceChat/Private/EOSAudioDevicePool.cpp
@@ -0,0 +1,197 @@
+// Copyright Epic Games, Inc. All Rights Reserved.
+
+#include "EOSAudioDevicePool.h"
+
+#if WITH_EOS_RTC
+
+#include "EOSVoiceChatLog.h"
+
+#include "eos_rtc.h"
+#include "eos_rtc_audio.h"
+
+namespace EOSAudioDevicePoolPrivate
+{
+	// State while handling a refresh request
+	// firing input `Delegate` during destroying
+	struct FRefreshAudioDevicesState
+	{
+		UE_NONCOPYABLE(FRefreshAudioDevicesState)
+
+		using FOnAudioDevicePoolRefreshAudioDevicesCompleteDelegate = FEOSAudioDevicePool::FOnAudioDevicePoolRefreshAudioDevicesCompleteDelegate;
+
+		explicit FRefreshAudioDevicesState(FOnAudioDevicePoolRefreshAudioDevicesCompleteDelegate InDelegate)
+		: Delegate(MoveTemp(InDelegate))
+		{
+		}
+
+		~FRefreshAudioDevicesState()
+		{
+			Delegate.ExecuteIfBound(MoveTemp(Result));
+		}
+
+		FVoiceChatResult Result = FVoiceChatResult::CreateSuccess();
+		const FOnAudioDevicePoolRefreshAudioDevicesCompleteDelegate Delegate;
+	};
+}
+
+FEOSAudioDevicePool::FEOSAudioDevicePool(EOS_HRTC& InEosRtcInterface)
+	: EosRtcInterface{ InEosRtcInterface }
+{
+}
+
+void FEOSAudioDevicePool::RefreshAudioDevices(const FOnAudioDevicePoolRefreshAudioDevicesCompleteDelegate& Delegate)
+{
+	// Delegate will be shared between several async calls and be invoked on the last one (in destructor as the last reference)
+	// with a success result if no one has been failed and failure otherwise
+	TSharedRef<EOSAudioDevicePoolPrivate::FRefreshAudioDevicesState> RefreshAudioDevicesState = MakeShared<EOSAudioDevicePoolPrivate::FRefreshAudioDevicesState>(Delegate);
+
+	RefreshAudioInputDevices(FOnAudioDevicePoolRefreshAudioInputDevicesCompleteDelegate::CreateLambda([RefreshAudioDevicesState](const FVoiceChatResult& InResult) -> void
+	{
+		if (!InResult.IsSuccess())
+		{
+			UE_LOG(LogEOSVoiceChat, Warning, TEXT("RefreshAudioInputDevicesCompletionDelegate failed, Result=[%s]"), *LexToString(InResult));
+
+			RefreshAudioDevicesState->Result = InResult;
+		}
+	}));
+	RefreshAudioOutputDevices(FOnAudioDevicePoolRefreshAudioOutputDevicesCompleteDelegate::CreateLambda([RefreshAudioDevicesState](const FVoiceChatResult& InResult) -> void
+	{
+		if (!InResult.IsSuccess())
+		{
+			UE_LOG(LogEOSVoiceChat, Warning, TEXT("RefreshAudioOutputDevicesCompletionDelegate failed, Result=[%s]"), *LexToString(InResult));
+
+			RefreshAudioDevicesState->Result = InResult;
+		}
+	}));
+}
+
+void FEOSAudioDevicePool::RefreshAudioInputDevices(const FOnAudioDevicePoolRefreshAudioInputDevicesCompleteDelegate& Delegate)
+{
+	CachedInputDeviceInfos = GetRtcInputDeviceInfos(DefaultInputDeviceInfoIdx);
+
+	UE_LOG(LogEOSVoiceChat, Verbose, TEXT("RefreshAudioInputDevices InputDeviceInfos=[%s] DefaultInputDeviceInfoIdx=[%d]"), *FString::JoinBy(CachedInputDeviceInfos, TEXT(", "), &FVoiceChatDeviceInfo::ToDebugString), DefaultInputDeviceInfoIdx);
+
+	Delegate.ExecuteIfBound(FVoiceChatResult::CreateSuccess());
+}
+
+void FEOSAudioDevicePool::RefreshAudioOutputDevices(const FOnAudioDevicePoolRefreshAudioOutputDevicesCompleteDelegate& Delegate)
+{
+	CachedOutputDeviceInfos = GetRtcOutputDeviceInfos(DefaultOutputDeviceInfoIdx);
+
+	UE_LOG(LogEOSVoiceChat, Verbose, TEXT("RefreshAudioOutputDevices OutputDeviceInfos=[%s] DefaultOutputDeviceInfoIdx=[%d]"), *FString::JoinBy(CachedOutputDeviceInfos, TEXT(", "), &FVoiceChatDeviceInfo::ToDebugString), DefaultOutputDeviceInfoIdx);
+
+	Delegate.ExecuteIfBound(FVoiceChatResult::CreateSuccess());
+}
+
+TArray<FVoiceChatDeviceInfo> FEOSAudioDevicePool::GetRtcInputDeviceInfos(int32& OutDefaultDeviceIdx) const
+{
+	TArray<FVoiceChatDeviceInfo> InputDeviceInfos;
+	EOS_HRTCAudio RTCAudioHandle = EOS_RTC_GetAudioInterface(EosRtcInterface);
+
+	EOS_RTCAudio_GetAudioInputDevicesCountOptions CountOptions = {};
+	CountOptions.ApiVersion = EOS_RTCAUDIO_GETAUDIOINPUTDEVICESCOUNT_API_LATEST;
+	static_assert(EOS_RTCAUDIO_GETAUDIOINPUTDEVICESCOUNT_API_LATEST == 1, "EOS_RTCAudio_GetAudioInputDevicesCountOptions updated, check new fields");
+
+	uint32_t Count = EOS_RTCAudio_GetAudioInputDevicesCount(RTCAudioHandle, &CountOptions);
+
+	for (uint32_t Index = 0; Index < Count; Index++)
+	{
+		EOS_RTCAudio_GetAudioInputDeviceByIndexOptions GetByIndexOptions = {};
+		GetByIndexOptions.ApiVersion = EOS_RTCAUDIO_GETAUDIOINPUTDEVICEBYINDEX_API_LATEST;
+		static_assert(EOS_RTCAUDIO_GETAUDIOINPUTDEVICEBYINDEX_API_LATEST == 1, "EOS_RTCAudio_GetAudioInputDeviceByIndexOptions updated, check new fields");
+		GetByIndexOptions.DeviceInfoIndex = Index;
+		if (const EOS_RTCAudio_AudioInputDeviceInfo* DeviceInfo = EOS_RTCAudio_GetAudioInputDeviceByIndex(RTCAudioHandle, &GetByIndexOptions))
+		{
+			FString DeviceName = UTF8_TO_TCHAR(DeviceInfo->DeviceName);
+			if (DeviceName != TEXT("Default Device"))
+			{
+				FVoiceChatDeviceInfo& InputDeviceInfo = InputDeviceInfos.Emplace_GetRef();
+				InputDeviceInfo.DisplayName = MoveTemp(DeviceName);
+				InputDeviceInfo.Id = UTF8_TO_TCHAR(DeviceInfo->DeviceId);
+				if (DeviceInfo->bDefaultDevice)
+				{
+					OutDefaultDeviceIdx = InputDeviceInfos.Num() - 1;
+				}
+			}
+		}
+		else
+		{
+			UE_LOG(LogEOSVoiceChat, Warning, TEXT("EOS_RTCAudio_GetAudioInputDeviceByIndex failed: DevicesInfo=nullptr"));
+		}
+	}
+
+	if (Count == 0)
+	{
+		UE_LOG(LogEOSVoiceChat, Warning, TEXT("EOS_RTCAudio_GetAudioInputDevicesCount failed: DevicesCount=0"));
+	}
+
+	return InputDeviceInfos;
+}
+
+TArray<FVoiceChatDeviceInfo> FEOSAudioDevicePool::GetRtcOutputDeviceInfos(int32& OutDefaultDeviceIdx) const
+{
+	OutDefaultDeviceIdx = -1;
+	TArray<FVoiceChatDeviceInfo> OutputDeviceInfos;
+	EOS_HRTCAudio RTCAudioHandle = EOS_RTC_GetAudioInterface(EosRtcInterface);
+
+	EOS_RTCAudio_GetAudioOutputDevicesCountOptions CountOptions = {};
+	CountOptions.ApiVersion = EOS_RTCAUDIO_GETAUDIOOUTPUTDEVICESCOUNT_API_LATEST;
+	static_assert(EOS_RTCAUDIO_GETAUDIOOUTPUTDEVICESCOUNT_API_LATEST == 1, "EOS_RTCAudio_GetAudioOutputDevicesCountOptions updated, check new fields");
+
+	uint32_t Count = EOS_RTCAudio_GetAudioOutputDevicesCount(RTCAudioHandle, &CountOptions);
+
+	for (uint32_t Index = 0; Index < Count; Index++)
+	{
+		EOS_RTCAudio_GetAudioOutputDeviceByIndexOptions GetByIndexOptions = {};
+		GetByIndexOptions.ApiVersion = EOS_RTCAUDIO_GETAUDIOOUTPUTDEVICEBYINDEX_API_LATEST;
+		static_assert(EOS_RTCAUDIO_GETAUDIOOUTPUTDEVICEBYINDEX_API_LATEST == 1, "EOS_RTCAudio_GetAudioOutputDeviceByIndexOptions updated, check new fields");
+		GetByIndexOptions.DeviceInfoIndex = Index;
+		if (const EOS_RTCAudio_AudioOutputDeviceInfo* DeviceInfo = EOS_RTCAudio_GetAudioOutputDeviceByIndex(RTCAudioHandle, &GetByIndexOptions))
+		{
+			FString DeviceName = UTF8_TO_TCHAR(DeviceInfo->DeviceName);
+			if (DeviceName != TEXT("Default Device"))
+			{
+				FVoiceChatDeviceInfo& InputDeviceInfo = OutputDeviceInfos.Emplace_GetRef();
+				InputDeviceInfo.DisplayName = MoveTemp(DeviceName);
+				InputDeviceInfo.Id = UTF8_TO_TCHAR(DeviceInfo->DeviceId);
+				if (DeviceInfo->bDefaultDevice)
+				{
+					OutDefaultDeviceIdx = OutputDeviceInfos.Num() - 1;
+				}
+			}
+		}
+		else
+		{
+			UE_LOG(LogEOSVoiceChat, Warning, TEXT("EOS_RTCAudio_GetAudioOutputDeviceByIndex failed: DevicesInfo=nullptr"));
+		}
+	}
+
+	if (Count == 0)
+	{
+		UE_LOG(LogEOSVoiceChat, Warning, TEXT("EOS_RTCAudio_GetAudioOutputDevicesCount failed: DevicesCount=0"));
+	}
+
+	return OutputDeviceInfos;
+}
+
+const TArray<FVoiceChatDeviceInfo>& FEOSAudioDevicePool::GetCachedInputDeviceInfos() const
+{
+	return CachedInputDeviceInfos;
+}
+
+const TArray<FVoiceChatDeviceInfo>& FEOSAudioDevicePool::GetCachedOutputDeviceInfos() const
+{
+	return CachedOutputDeviceInfos;
+}
+
+int32 FEOSAudioDevicePool::GetDefaultInputDeviceInfoIdx() const
+{
+	return DefaultInputDeviceInfoIdx;
+}
+
+int32 FEOSAudioDevicePool::GetDefaultOutputDeviceInfoIdx() const
+{
+	return DefaultOutputDeviceInfoIdx;
+}
+
+#endif // WITH_EOS_RTC
diff --git a/Engine/Plugins/Online/VoiceChat/EOSVoiceChat/Source/EOSVoiceChat/Private/EOSAudioDevicePool.h b/Engine/Plugins/Online/VoiceChat/EOSVoiceChat/Source/EOSVoiceChat/Private/EOSAudioDevicePool.h
new file mode 100644
--- /dev/null
+++ b/Engine/Plugins/Online/VoiceChat/EOSVoiceChat/Source/EOSVoiceChat/Private/EOSAudioDevicePool.h
@@ -0,0 +1,45 @@
+// Copyright Epic Games, Inc. All Rights Reserved.
+
+#pragma once
+
+#include "EOSShared.h"
+
+#if WITH_EOS_RTC
+
+#include "VoiceChat.h"
+
+#include "eos_rtc_types.h"
+
+class FEOSAudioDevicePool final
+{
+public:
+	explicit FEOSAudioDevicePool(EOS_HRTC& InEosRtcInterface);
+
+	// forward refreshing both input and output devices
+	DECLARE_DELEGATE_OneParam(FOnAudioDevicePoolRefreshAudioDevicesCompleteDelegate, const FVoiceChatResult& /* Result */);
+	void RefreshAudioDevices(const FOnAudioDevicePoolRefreshAudioDevicesCompleteDelegate& Delegate);
+
+	DECLARE_DELEGATE_OneParam(FOnAudioDevicePoolRefreshAudioInputDevicesCompleteDelegate, const FVoiceChatResult& /* Result */);
+	void RefreshAudioInputDevices(const FOnAudioDevicePoolRefreshAudioInputDevicesCompleteDelegate& Delegate);
+
+	DECLARE_DELEGATE_OneParam(FOnAudioDevicePoolRefreshAudioOutputDevicesCompleteDelegate, const FVoiceChatResult& /* Result */);
+	void RefreshAudioOutputDevices(const FOnAudioDevicePoolRefreshAudioOutputDevicesCompleteDelegate& Delegate);
+
+	const TArray<FVoiceChatDeviceInfo>& GetCachedInputDeviceInfos() const;
+	const TArray<FVoiceChatDeviceInfo>& GetCachedOutputDeviceInfos() const;
+
+	int32 GetDefaultInputDeviceInfoIdx() const;
+	int32 GetDefaultOutputDeviceInfoIdx() const;
+private:
+	TArray<FVoiceChatDeviceInfo> GetRtcInputDeviceInfos(int32& OutDefaultDeviceIdx) const;
+	TArray<FVoiceChatDeviceInfo> GetRtcOutputDeviceInfos(int32& OutDefaultDeviceIdx) const;
+
+	EOS_HRTC& EosRtcInterface;
+
+	int32 DefaultInputDeviceInfoIdx = -1;
+	int32 DefaultOutputDeviceInfoIdx = -1;
+	TArray<FVoiceChatDeviceInfo> CachedInputDeviceInfos;
+	TArray<FVoiceChatDeviceInfo> CachedOutputDeviceInfos;
+};
+
+#endif // WITH_EOS_RTC
diff --git a/Engine/Plugins/Online/VoiceChat/EOSVoiceChat/Source/EOSVoiceChat/Private/EOSVoiceChat.cpp b/Engine/Plugins/Online/VoiceChat/EOSVoiceChat/Source/EOSVoiceChat/Private/EOSVoiceChat.cpp
--- a/Engine/Plugins/Online/VoiceChat/EOSVoiceChat/Source/EOSVoiceChat/Private/EOSVoiceChat.cpp
+++ b/Engine/Plugins/Online/VoiceChat/EOSVoiceChat/Source/EOSVoiceChat/Private/EOSVoiceChat.cpp
@@ -12,9 +12,11 @@
 #include "Stats/Stats.h"
 
 #include "EOSShared.h"
+#include "EOSVoiceChatLog.h"
 #include "EOSVoiceChatErrors.h"
 #include "EOSVoiceChatModule.h"
 #include "EOSVoiceChatUser.h"
+#include "EOSAudioDevicePool.h"
 #include "IEOSSDKManager.h"
 #include "VoiceChatErrors.h"
 #include "VoiceChatResult.h"
@@ -931,105 +933,17 @@ void FEOSVoiceChat::OnAudioDevicesChangedStatic(const EOS_RTCAudio_AudioDevicesC
 
 void FEOSVoiceChat::OnAudioDevicesChanged()
 {
-	InitSession.CachedInputDeviceInfos = GetRtcInputDeviceInfos(InitSession.DefaultInputDeviceInfoIdx);
-	InitSession.CachedOutputDeviceInfos = GetRtcOutputDeviceInfos(InitSession.DefaultOutputDeviceInfoIdx);
-
-	UE_LOG(LogEOSVoiceChat, Verbose, TEXT("OnAudioDevicesChanged InputDeviceInfos=[%s] DefaultInputDeviceInfoIdx=%d"), *FString::JoinBy(InitSession.CachedInputDeviceInfos, TEXT(", "), &FVoiceChatDeviceInfo::ToDebugString), InitSession.DefaultInputDeviceInfoIdx);
-	UE_LOG(LogEOSVoiceChat, Verbose, TEXT("OnAudioDevicesChanged OutputDeviceInfos=[%s] DefaultOutputDeviceInfoIdx=%d"), *FString::JoinBy(InitSession.CachedOutputDeviceInfos, TEXT(", "), &FVoiceChatDeviceInfo::ToDebugString), InitSession.DefaultOutputDeviceInfoIdx);
-
-	OnVoiceChatAvailableAudioDevicesChangedDelegate.Broadcast();
-}
-
-TArray<FVoiceChatDeviceInfo> FEOSVoiceChat::GetRtcInputDeviceInfos(int32& OutDefaultDeviceIdx) const
-{
-	OutDefaultDeviceIdx = -1;
-	TArray<FVoiceChatDeviceInfo> InputDeviceInfos;
-	EOS_HRTCAudio RTCAudioHandle = EOS_RTC_GetAudioInterface(InitSession.EosRtcInterface);
-
-	EOS_RTCAudio_GetAudioInputDevicesCountOptions CountOptions = {};
-	CountOptions.ApiVersion = EOS_RTCAUDIO_GETAUDIOINPUTDEVICESCOUNT_API_LATEST;
- 	static_assert(EOS_RTCAUDIO_GETAUDIOINPUTDEVICESCOUNT_API_LATEST == 1, "EOS_RTCAudio_GetAudioInputDevicesCountOptions updated, check new fields");
-
-	uint32_t Count = EOS_RTCAudio_GetAudioInputDevicesCount(RTCAudioHandle, &CountOptions);
-
-	for (uint32_t Index = 0; Index < Count; Index++)
-	{
-		EOS_RTCAudio_GetAudioInputDeviceByIndexOptions GetByIndexOptions = {};
-		GetByIndexOptions.ApiVersion = EOS_RTCAUDIO_GETAUDIOINPUTDEVICEBYINDEX_API_LATEST;
-		static_assert(EOS_RTCAUDIO_GETAUDIOINPUTDEVICEBYINDEX_API_LATEST == 1, "EOS_RTCAudio_GetAudioInputDeviceByIndexOptions updated, check new fields");
-		GetByIndexOptions.DeviceInfoIndex = Index;
-		if (const EOS_RTCAudio_AudioInputDeviceInfo* DeviceInfo = EOS_RTCAudio_GetAudioInputDeviceByIndex(RTCAudioHandle, &GetByIndexOptions))
-		{
-			FString DeviceName = UTF8_TO_TCHAR(DeviceInfo->DeviceName);
-			if (DeviceName != TEXT("Default Device"))
-			{
-				FVoiceChatDeviceInfo& InputDeviceInfo = InputDeviceInfos.Emplace_GetRef();
-				InputDeviceInfo.DisplayName = MoveTemp(DeviceName);
-				InputDeviceInfo.Id = UTF8_TO_TCHAR(DeviceInfo->DeviceId);
-				if (DeviceInfo->bDefaultDevice)
-				{
-					OutDefaultDeviceIdx = InputDeviceInfos.Num() - 1;
-				}
-			}
-		}
-		else
-		{
-			UE_LOG(LogEOSVoiceChat, Warning, TEXT("EOS_RTCAudio_GetAudioInputDeviceByIndex failed: DevicesInfo=nullptr"));
-		}
-	}
-
-	if (Count == 0)
+	InitSession.EosAudioDevicePool->RefreshAudioDevices(FEOSAudioDevicePool::FOnAudioDevicePoolRefreshAudioDevicesCompleteDelegate::CreateLambda([WeakThis = AsWeak()](const FVoiceChatResult& Result) -> void
 	{
-		UE_LOG(LogEOSVoiceChat, Warning, TEXT("EOS_RTCAudio_GetAudioInputDevicesCount failed: DevicesCount=0"));
-	}
+		CHECKPIN();
 
-	return InputDeviceInfos;
-}
-
-TArray<FVoiceChatDeviceInfo> FEOSVoiceChat::GetRtcOutputDeviceInfos(int32& OutDefaultDeviceIdx) const
-{
-	OutDefaultDeviceIdx = -1;
-	TArray<FVoiceChatDeviceInfo> OutputDeviceInfos;
-	EOS_HRTCAudio RTCAudioHandle = EOS_RTC_GetAudioInterface(InitSession.EosRtcInterface);
-
-	EOS_RTCAudio_GetAudioOutputDevicesCountOptions CountOptions = {};
-	CountOptions.ApiVersion = EOS_RTCAUDIO_GETAUDIOOUTPUTDEVICESCOUNT_API_LATEST;
-	static_assert(EOS_RTCAUDIO_GETAUDIOOUTPUTDEVICESCOUNT_API_LATEST == 1, "EOS_RTCAudio_GetAudioOutputDevicesCountOptions updated, check new fields");
-
-	uint32_t Count = EOS_RTCAudio_GetAudioOutputDevicesCount(RTCAudioHandle, &CountOptions);
-
-	for (uint32_t Index = 0; Index < Count; Index++)
-	{
-		EOS_RTCAudio_GetAudioOutputDeviceByIndexOptions GetByIndexOptions = {};
-		GetByIndexOptions.ApiVersion = EOS_RTCAUDIO_GETAUDIOOUTPUTDEVICEBYINDEX_API_LATEST;
-		static_assert(EOS_RTCAUDIO_GETAUDIOOUTPUTDEVICEBYINDEX_API_LATEST == 1, "EOS_RTCAudio_GetAudioOutputDeviceByIndexOptions updated, check new fields");
-		GetByIndexOptions.DeviceInfoIndex = Index;
-		if (const EOS_RTCAudio_AudioOutputDeviceInfo* DeviceInfo = EOS_RTCAudio_GetAudioOutputDeviceByIndex(RTCAudioHandle, &GetByIndexOptions))
+		if (!Result.IsSuccess())
 		{
-			FString DeviceName = UTF8_TO_TCHAR(DeviceInfo->DeviceName);
-			if (DeviceName != TEXT("Default Device"))
-			{
-				FVoiceChatDeviceInfo& InputDeviceInfo = OutputDeviceInfos.Emplace_GetRef();
-				InputDeviceInfo.DisplayName = MoveTemp(DeviceName);
-				InputDeviceInfo.Id = UTF8_TO_TCHAR(DeviceInfo->DeviceId);
-				if (DeviceInfo->bDefaultDevice)
-				{
-					OutDefaultDeviceIdx = OutputDeviceInfos.Num() - 1;
-				}
-			}
+			UE_LOG(LogEOSVoiceChat, Warning, TEXT("OnAudioDevicesChanged RefreshAudioDevicesCompletionDelegate failed, Result=[%s]"), *LexToString(Result));
 		}
-		else
-		{
-			UE_LOG(LogEOSVoiceChat, Warning, TEXT("EOS_RTCAudio_GetAudioOutputDeviceByIndex failed: DevicesInfo=nullptr"));
-		}
-	}
-
-	if (Count == 0)
-	{
-		UE_LOG(LogEOSVoiceChat, Warning, TEXT("EOS_RTCAudio_GetAudioOutputDevicesCount failed: DevicesCount=0"));
-	}
 
-	return OutputDeviceInfos;
+		StrongThis->OnVoiceChatAvailableAudioDevicesChangedDelegate.Broadcast();
+	}));
 }
 
 FEOSVoiceChatUser& FEOSVoiceChat::GetVoiceChatUser()
@@ -1211,6 +1125,11 @@ bool FEOSVoiceChat::Exec(UWorld* InWorld, const TCHAR* Cmd, FOutputDevice& Ar)
 	return false;
 }
 
+FEOSVoiceChat::FInitSession::FInitSession()
+	: EosAudioDevicePool{ MakeUnique<FEOSAudioDevicePool>(EosRtcInterface) }
+{
+}
+
 IEOSPlatformHandlePtr FEOSVoiceChat::EOSPlatformCreate(EOS_Platform_Options& PlatformOptions)
 {
 	return SDKManager.CreatePlatform(PlatformOptions);
diff --git a/Engine/Plugins/Online/VoiceChat/EOSVoiceChat/Source/EOSVoiceChat/Private/EOSVoiceChatFactory.cpp b/Engine/Plugins/Online/VoiceChat/EOSVoiceChat/Source/EOSVoiceChat/Private/EOSVoiceChatFactory.cpp
--- a/Engine/Plugins/Online/VoiceChat/EOSVoiceChat/Source/EOSVoiceChat/Private/EOSVoiceChatFactory.cpp
+++ b/Engine/Plugins/Online/VoiceChat/EOSVoiceChat/Source/EOSVoiceChat/Private/EOSVoiceChatFactory.cpp
@@ -6,6 +6,7 @@
 
 #include "Misc/Parse.h"
 #include "IEOSSDKManager.h"
+#include "EOSVoiceChatLog.h"
 
 #include COMPILED_PLATFORM_HEADER(EOSVoiceChat.h)
 
diff --git a/Engine/Plugins/Online/VoiceChat/EOSVoiceChat/Source/EOSVoiceChat/Private/EOSVoiceChatUser.cpp b/Engine/Plugins/Online/VoiceChat/EOSVoiceChat/Source/EOSVoiceChat/Private/EOSVoiceChatUser.cpp
--- a/Engine/Plugins/Online/VoiceChat/EOSVoiceChat/Source/EOSVoiceChat/Private/EOSVoiceChatUser.cpp
+++ b/Engine/Plugins/Online/VoiceChat/EOSVoiceChat/Source/EOSVoiceChat/Private/EOSVoiceChatUser.cpp
@@ -14,6 +14,7 @@
 #include "EOSVoiceChatModule.h"
 #include "EOSVoiceChatTypes.h"
 #include "VoiceChatErrors.h"
+#include "EOSAudioDevicePool.h"
 
 #include "eos_common.h"
 #include "eos_lobby.h"
@@ -221,12 +222,12 @@ bool FEOSVoiceChatUser::GetAudioOutputDeviceMuted() const
 
 TArray<FVoiceChatDeviceInfo> FEOSVoiceChatUser::GetAvailableInputDeviceInfos() const
 {
-	return EOSVoiceChat.InitSession.CachedInputDeviceInfos;
+	return EOSVoiceChat.InitSession.EosAudioDevicePool->GetCachedInputDeviceInfos();
 }
 
 TArray<FVoiceChatDeviceInfo> FEOSVoiceChatUser::GetAvailableOutputDeviceInfos() const
 {
-	return EOSVoiceChat.InitSession.CachedOutputDeviceInfos;
+	return EOSVoiceChat.InitSession.EosAudioDevicePool->GetCachedOutputDeviceInfos();
 }
 
 void FEOSVoiceChatUser::SetInputDeviceId(const FString& InputDeviceId)
@@ -305,9 +306,9 @@ FVoiceChatDeviceInfo FEOSVoiceChatUser::GetDefaultInputDeviceInfo() const
 {
 	FVoiceChatDeviceInfo DefaultDeviceInfo;
 
-	if (EOSVoiceChat.InitSession.CachedInputDeviceInfos.IsValidIndex(EOSVoiceChat.InitSession.DefaultInputDeviceInfoIdx))
+	if (EOSVoiceChat.InitSession.EosAudioDevicePool->GetCachedInputDeviceInfos().IsValidIndex(EOSVoiceChat.InitSession.EosAudioDevicePool->GetDefaultInputDeviceInfoIdx()))
 	{
-		DefaultDeviceInfo = EOSVoiceChat.InitSession.CachedInputDeviceInfos[EOSVoiceChat.InitSession.DefaultInputDeviceInfoIdx];
+		DefaultDeviceInfo = EOSVoiceChat.InitSession.EosAudioDevicePool->GetCachedInputDeviceInfos()[EOSVoiceChat.InitSession.EosAudioDevicePool->GetDefaultInputDeviceInfoIdx()];
 	}
 
 	return DefaultDeviceInfo;
@@ -317,9 +318,9 @@ FVoiceChatDeviceInfo FEOSVoiceChatUser::GetDefaultOutputDeviceInfo() const
 {
 	FVoiceChatDeviceInfo DefaultDeviceInfo;
 
-	if (EOSVoiceChat.InitSession.CachedOutputDeviceInfos.IsValidIndex(EOSVoiceChat.InitSession.DefaultOutputDeviceInfoIdx))
+	if (EOSVoiceChat.InitSession.EosAudioDevicePool->GetCachedOutputDeviceInfos().IsValidIndex(EOSVoiceChat.InitSession.EosAudioDevicePool->GetDefaultOutputDeviceInfoIdx()))
 	{
-		DefaultDeviceInfo = EOSVoiceChat.InitSession.CachedOutputDeviceInfos[EOSVoiceChat.InitSession.DefaultOutputDeviceInfoIdx];
+		DefaultDeviceInfo = EOSVoiceChat.InitSession.EosAudioDevicePool->GetCachedOutputDeviceInfos()[EOSVoiceChat.InitSession.EosAudioDevicePool->GetDefaultOutputDeviceInfoIdx()];
 	}
 	 
 	return DefaultDeviceInfo;
diff --git a/Engine/Plugins/Online/VoiceChat/EOSVoiceChat/Source/EOSVoiceChat/Public/EOSVoiceChat.h b/Engine/Plugins/Online/VoiceChat/EOSVoiceChat/Source/EOSVoiceChat/Public/EOSVoiceChat.h
--- a/Engine/Plugins/Online/VoiceChat/EOSVoiceChat/Source/EOSVoiceChat/Public/EOSVoiceChat.h
+++ b/Engine/Plugins/Online/VoiceChat/EOSVoiceChat/Source/EOSVoiceChat/Public/EOSVoiceChat.h
@@ -14,8 +14,6 @@
 #include "eos_rtc_audio_types.h"
 #include "eos_lobby_types.h"
 
-DECLARE_LOG_CATEGORY_EXTERN(LogEOSVoiceChat, Log, All);
-
 EOSVOICECHAT_API FVoiceChatResult ResultFromEOSResult(const EOS_EResult EosResult);
 
 EOSVOICECHAT_API const TCHAR* LexToString(EOS_ERTCAudioInputStatus Status);
@@ -45,6 +43,7 @@ public:
 
 DECLARE_MULTICAST_DELEGATE_ThreeParams(FOnVoiceChatDataReceivedDelegate, const FString& /* ChannelName */, const FString& /* PlayerName */, TArrayView<const uint8> /* Buffer */);
 
+class FEOSAudioDevicePool;
 class FEOSVoiceChatUser;
 class IEOSSDKManager;
 using IEOSPlatformHandlePtr = TSharedPtr<class IEOSPlatformHandle, ESPMode::ThreadSafe>;
@@ -168,6 +167,8 @@ protected:
 
 	struct FInitSession
 	{
+		FInitSession();
+
 		EInitializationState State = EInitializationState::Uninitialized;
 
 		TArray<FOnVoiceChatUninitializeCompleteDelegate> UninitializeCompleteDelegates;
@@ -178,10 +179,7 @@ protected:
 
 		EOS_NotificationId OnAudioDevicesChangedNotificationId = EOS_INVALID_NOTIFICATIONID;
 
-		int32 DefaultInputDeviceInfoIdx = -1;
-		int32 DefaultOutputDeviceInfoIdx = -1;
-		TArray<FVoiceChatDeviceInfo> CachedInputDeviceInfos;
-		TArray<FVoiceChatDeviceInfo> CachedOutputDeviceInfos;
+		TUniquePtr<FEOSAudioDevicePool> EosAudioDevicePool;
 	};
 	FInitSession InitSession;
 
@@ -191,9 +189,6 @@ protected:
 	static void EOS_CALL OnAudioDevicesChangedStatic(const EOS_RTCAudio_AudioDevicesChangedCallbackInfo* CallbackInfo);
 	void OnAudioDevicesChanged();
 
-	TArray<FVoiceChatDeviceInfo> GetRtcInputDeviceInfos(int32& OutDefaultDeviceIdx) const;
-	TArray<FVoiceChatDeviceInfo> GetRtcOutputDeviceInfos(int32& OutDefaultDeviceIdx) const;
-
 	// The current state of the connection.
 	enum class EConnectionState
 	{
diff --git a/Engine/Plugins/Online/VoiceChat/EOSVoiceChat/Source/EOSVoiceChat/Public/EOSVoiceChatLog.h b/Engine/Plugins/Online/VoiceChat/EOSVoiceChat/Source/EOSVoiceChat/Public/EOSVoiceChatLog.h
new file mode 100644
--- /dev/null
+++ b/Engine/Plugins/Online/VoiceChat/EOSVoiceChat/Source/EOSVoiceChat/Public/EOSVoiceChatLog.h
@@ -0,0 +1,21 @@
+// Copyright Epic Games, Inc. All Rights Reserved.
+
+#pragma once
+
+#include "Logging/LogMacros.h"
+
+#if WITH_EOS_RTC
+
+DECLARE_LOG_CATEGORY_EXTERN(LogEOSVoiceChat, Log, All);
+
+#define EOSVOICECHATUSER_LOG(EOSVoiceChatLogLevel, EOSVoiceChatFormatStr, ...) \
+{ \
+	UE_LOG(LogEOSVoiceChat, EOSVoiceChatLogLevel, TEXT("[%p] ") EOSVoiceChatFormatStr, (void*)this, ##__VA_ARGS__); \
+}
+
+#define EOSVOICECHATUSER_CLOG(Condition, EOSVoiceChatLogLevel, EOSVoiceChatFormatStr, ...) \
+{ \
+	UE_CLOG(Condition, LogEOSVoiceChat, EOSVoiceChatLogLevel, TEXT("[%p] ") EOSVoiceChatFormatStr, (void*)this, ##__VA_ARGS__); \
+}
+
+#endif // WITH_EOS_RTC
diff --git a/Engine/Plugins/Online/VoiceChat/EOSVoiceChat/Source/EOSVoiceChat/Public/EOSVoiceChatUser.h b/Engine/Plugins/Online/VoiceChat/EOSVoiceChat/Source/EOSVoiceChat/Public/EOSVoiceChatUser.h
--- a/Engine/Plugins/Online/VoiceChat/EOSVoiceChat/Source/EOSVoiceChat/Public/EOSVoiceChatUser.h
+++ b/Engine/Plugins/Online/VoiceChat/EOSVoiceChat/Source/EOSVoiceChat/Public/EOSVoiceChatUser.h
@@ -3,6 +3,7 @@
 #pragma once
 
 #include "EOSVoiceChat.h"
+#include "EOSVoiceChatLog.h"
 
 #if WITH_EOS_RTC
 
@@ -10,16 +11,6 @@
 #include "eos_rtc_audio_types.h"
 #include "eos_types.h"
 
-#define EOSVOICECHATUSER_LOG(EOSVoiceChatLogLevel, EOSVoiceChatFormatStr, ...) \
-{ \
-	UE_LOG(LogEOSVoiceChat, EOSVoiceChatLogLevel, TEXT("[%p] ") EOSVoiceChatFormatStr, (void*)this, ##__VA_ARGS__); \
-}
-
-#define EOSVOICECHATUSER_CLOG(Condition, EOSVoiceChatLogLevel, EOSVoiceChatFormatStr, ...) \
-{ \
-	UE_CLOG(Condition, LogEOSVoiceChat, EOSVoiceChatLogLevel, TEXT("[%p] ") EOSVoiceChatFormatStr, (void*)this, ##__VA_ARGS__); \
-}
-
 class EOSVOICECHAT_API FEOSVoiceChatUser : public TSharedFromThis<FEOSVoiceChatUser, ESPMode::ThreadSafe>, public IVoiceChatUser
 {
 public:
