- Fix crash due to modifying container being iterated over

diff --git a/Engine/Plugins/Online/VoiceChat/EOSVoiceChat/Source/EOSVoiceChat/Private/EOSVoiceChatUser.cpp b/Engine/Plugins/Online/VoiceChat/EOSVoiceChat/Source/EOSVoiceChat/Private/EOSVoiceChatUser.cpp
--- a/Engine/Plugins/Online/VoiceChat/EOSVoiceChat/Source/EOSVoiceChat/Private/EOSVoiceChatUser.cpp
+++ b/Engine/Plugins/Online/VoiceChat/EOSVoiceChat/Source/EOSVoiceChat/Private/EOSVoiceChatUser.cpp
@@ -1568,7 +1568,7 @@ void FEOSVoiceChatUser::ClearLoginSession()
 {
 	for (TPair<FString, FChannelSession> Pair : LoginSession.ChannelSessions)
 	{
-		RemoveChannelSession(Pair.Key);
+		UnbindChannelCallbacks(Pair.Value);
 	}
 	LoginSession = FLoginSession();
 }
