Use EOSAudioDevicePool as a shared instance

1) To be able to hold itself internally.

Test plan:
1) Run Incremental. https://horde.devtools.epicgames.com/job/629f492a410219bf1b270814

diff --git a/Engine/Plugins/Online/VoiceChat/EOSVoiceChat/Source/EOSVoiceChat/Private/EOSAudioDevicePool.h b/Engine/Plugins/Online/VoiceChat/EOSVoiceChat/Source/EOSVoiceChat/Private/EOSAudioDevicePool.h
--- a/Engine/Plugins/Online/VoiceChat/EOSVoiceChat/Source/EOSVoiceChat/Private/EOSAudioDevicePool.h
+++ b/Engine/Plugins/Online/VoiceChat/EOSVoiceChat/Source/EOSVoiceChat/Private/EOSAudioDevicePool.h
@@ -10,7 +10,7 @@
 
 #include "eos_rtc_types.h"
 
-class FEOSAudioDevicePool final
+class FEOSAudioDevicePool final : public TSharedFromThis<FEOSAudioDevicePool>
 {
 public:
 	explicit FEOSAudioDevicePool(EOS_HRTC& InEosRtcInterface);
diff --git a/Engine/Plugins/Online/VoiceChat/EOSVoiceChat/Source/EOSVoiceChat/Private/EOSVoiceChat.cpp b/Engine/Plugins/Online/VoiceChat/EOSVoiceChat/Source/EOSVoiceChat/Private/EOSVoiceChat.cpp
--- a/Engine/Plugins/Online/VoiceChat/EOSVoiceChat/Source/EOSVoiceChat/Private/EOSVoiceChat.cpp
+++ b/Engine/Plugins/Online/VoiceChat/EOSVoiceChat/Source/EOSVoiceChat/Private/EOSVoiceChat.cpp
@@ -1017,7 +1017,7 @@ bool FEOSVoiceChat::Exec(UWorld* InWorld, const TCHAR* Cmd, FOutputDevice& Ar)
 }
 
 FEOSVoiceChat::FInitSession::FInitSession()
-	: EosAudioDevicePool{ MakeUnique<FEOSAudioDevicePool>(EosRtcInterface) }
+	: EosAudioDevicePool{ MakeShared<FEOSAudioDevicePool>(EosRtcInterface) }
 {
 }
 
diff --git a/Engine/Plugins/Online/VoiceChat/EOSVoiceChat/Source/EOSVoiceChat/Public/EOSVoiceChat.h b/Engine/Plugins/Online/VoiceChat/EOSVoiceChat/Source/EOSVoiceChat/Public/EOSVoiceChat.h
--- a/Engine/Plugins/Online/VoiceChat/EOSVoiceChat/Source/EOSVoiceChat/Public/EOSVoiceChat.h
+++ b/Engine/Plugins/Online/VoiceChat/EOSVoiceChat/Source/EOSVoiceChat/Public/EOSVoiceChat.h
@@ -177,7 +177,7 @@ protected:
 
 		EOS_NotificationId OnAudioDevicesChangedNotificationId = EOS_INVALID_NOTIFICATIONID;
 
-		TUniquePtr<FEOSAudioDevicePool> EosAudioDevicePool;
+		TSharedRef<FEOSAudioDevicePool> EosAudioDevicePool;
 	};
 	FInitSession InitSession;
 
