Call RtcRegister|UnregisterUser with tracking its result

1) Added completion delegates to rtcRegister|UnregisterUser to be able to track their result and updated to return the actual result.
3) Updated usage places.

Test plan:
1) Run Incremental. https://horde.devtools.epicgames.com/job/62a23524415407fc4cd8f78d

diff --git a/Engine/Plugins/Online/VoiceChat/EOSVoiceChat/Source/EOSVoiceChat/Private/EOSVoiceChatUser.cpp b/Engine/Plugins/Online/VoiceChat/EOSVoiceChat/Source/EOSVoiceChat/Private/EOSVoiceChatUser.cpp
--- a/Engine/Plugins/Online/VoiceChat/EOSVoiceChat/Source/EOSVoiceChat/Private/EOSVoiceChatUser.cpp
+++ b/Engine/Plugins/Online/VoiceChat/EOSVoiceChat/Source/EOSVoiceChat/Private/EOSVoiceChatUser.cpp
@@ -1022,7 +1022,7 @@ bool FEOSVoiceChatUser::RemoveLobbyRoom(const FString& LobbyId)
 	return false;
 }
 
-void FEOSVoiceChatUser::RtcRegisterUser(const FString& UserId)
+void FEOSVoiceChatUser::RtcRegisterUser(const FString& UserId, const FOnVoiceChatUserRtcRegisterUserCompleteDelegate& Delegate)
 {
 	const FTCHARToUTF8 Utf8UserId(*UserId);
 
@@ -1033,10 +1033,16 @@ void FEOSVoiceChatUser::RtcRegisterUser(const FString& UserId)
 
 	CSV_SCOPED_TIMING_STAT_EXCLUSIVE(EOSVoiceChat);
 	QUICK_SCOPE_CYCLE_COUNTER(EOS_RTC_RegisterAudioUser);
-	EOS_RTCAudio_RegisterPlatformAudioUser(EOS_RTC_GetAudioInterface(GetRtcInterface()), &Options);
+	const EOS_EResult EosResult = EOS_RTCAudio_RegisterPlatformAudioUser(EOS_RTC_GetAudioInterface(GetRtcInterface()), &Options);
+	if (EosResult != EOS_EResult::EOS_Success)
+	{
+		EOSVOICECHATUSER_LOG(Warning, TEXT("RtcRegisterUser failed: %s"), *LexToString(EosResult));
+	}
+
+	Delegate.ExecuteIfBound(EosResult);
 }
 
-void FEOSVoiceChatUser::RtcUnregisterUser(const FString& UserId)
+void FEOSVoiceChatUser::RtcUnregisterUser(const FString& UserId, const FOnVoiceChatUserRtcUnregisterUserCompleteDelegate& Delegate)
 {
 	const FTCHARToUTF8 Utf8UserId(*UserId);
 
@@ -1047,7 +1053,13 @@ void FEOSVoiceChatUser::RtcUnregisterUser(const FString& UserId)
 
 	CSV_SCOPED_TIMING_STAT_EXCLUSIVE(EOSVoiceChat);
 	QUICK_SCOPE_CYCLE_COUNTER(EOS_RTC_UnregisterAudioUser);
-	EOS_RTCAudio_UnregisterPlatformAudioUser(EOS_RTC_GetAudioInterface(GetRtcInterface()), &Options);
+	const EOS_EResult EosResult = EOS_RTCAudio_UnregisterPlatformAudioUser(EOS_RTC_GetAudioInterface(GetRtcInterface()), &Options);
+	if (EosResult != EOS_EResult::EOS_Success)
+	{
+		EOSVOICECHATUSER_LOG(Warning, TEXT("RtcUnregisterUser failed: %s"), *LexToString(EosResult));
+	}
+
+	Delegate.ExecuteIfBound(EosResult);
 }
 
 void FEOSVoiceChatUser::SetHardwareAECEnabled(bool bEnabled)
diff --git a/Engine/Plugins/Online/VoiceChat/EOSVoiceChat/Source/EOSVoiceChat/Public/EOSVoiceChatUser.h b/Engine/Plugins/Online/VoiceChat/EOSVoiceChat/Source/EOSVoiceChat/Public/EOSVoiceChatUser.h
--- a/Engine/Plugins/Online/VoiceChat/EOSVoiceChat/Source/EOSVoiceChat/Public/EOSVoiceChatUser.h
+++ b/Engine/Plugins/Online/VoiceChat/EOSVoiceChat/Source/EOSVoiceChat/Public/EOSVoiceChatUser.h
@@ -305,8 +305,12 @@ protected:
 	void LeaveChannelInternal(const FString& ChannelName, const FOnVoiceChatChannelLeaveCompleteDelegate& Delegate);
 	void LogoutInternal(const FOnVoiceChatLogoutCompleteDelegate& Delegate);
 	void ClearLoginSession();
-	void RtcRegisterUser(const FString& UserId);
-	void RtcUnregisterUser(const FString& UserId);
+
+	DECLARE_DELEGATE_OneParam(FOnVoiceChatUserRtcRegisterUserCompleteDelegate, const EOS_EResult /* Result */);
+	void RtcRegisterUser(const FString& UserId, const FOnVoiceChatUserRtcRegisterUserCompleteDelegate& Delegate);
+	DECLARE_DELEGATE_OneParam(FOnVoiceChatUserRtcUnregisterUserCompleteDelegate, const EOS_EResult /* Result */);
+	void RtcUnregisterUser(const FString& UserId, const FOnVoiceChatUserRtcUnregisterUserCompleteDelegate& Delegate);
+	
 	void SetHardwareAECEnabled(bool bEnabled);
     
 	// EOS operation callbacks
