- Fix accidentally unbinding channel delegates with a _copy_ of the FChannelSession, causing some asserts to fire in the destructor of the real one.

diff --git a/Engine/Plugins/Online/VoiceChat/EOSVoiceChat/Source/EOSVoiceChat/Private/EOSVoiceChatUser.cpp b/Engine/Plugins/Online/VoiceChat/EOSVoiceChat/Source/EOSVoiceChat/Private/EOSVoiceChatUser.cpp
--- a/Engine/Plugins/Online/VoiceChat/EOSVoiceChat/Source/EOSVoiceChat/Private/EOSVoiceChatUser.cpp
+++ b/Engine/Plugins/Online/VoiceChat/EOSVoiceChat/Source/EOSVoiceChat/Private/EOSVoiceChatUser.cpp
@@ -1578,7 +1578,7 @@ void FEOSVoiceChatUser::LogoutInternal(const FOnVoiceChatLogoutCompleteDelegate&
 
 void FEOSVoiceChatUser::ClearLoginSession()
 {
-	for (TPair<FString, FChannelSession> Pair : LoginSession.ChannelSessions)
+	for (TPair<FString, FChannelSession>& Pair : LoginSession.ChannelSessions)
 	{
 		UnbindChannelCallbacks(Pair.Value);
 	}
