Fix EOS voice crash after its reinitialization

0) EOSDevicePool holds a ref to a parent field which is invalidated after copying/moving.
1) Disabled FInitSession copying/moving and added reset to clear that object state.

Test plan:
1) Run incremental https://horde.devtools.epicgames.com/job/62ac5643fc406f218cc770ee.
2) Run Editor with voice unit/init + EOSDevicePool::Refresh. It isn't crashed anymore.

diff --git a/Engine/Plugins/Online/VoiceChat/EOSVoiceChat/Source/EOSVoiceChat/Private/EOSVoiceChat.cpp b/Engine/Plugins/Online/VoiceChat/EOSVoiceChat/Source/EOSVoiceChat/Private/EOSVoiceChat.cpp
--- a/Engine/Plugins/Online/VoiceChat/EOSVoiceChat/Source/EOSVoiceChat/Private/EOSVoiceChat.cpp
+++ b/Engine/Plugins/Online/VoiceChat/EOSVoiceChat/Source/EOSVoiceChat/Private/EOSVoiceChat.cpp
@@ -166,7 +166,7 @@ void FEOSVoiceChat::Initialize(const FOnVoiceChatInitializeCompleteDelegate& Ini
 					if (!InitSession.EosPlatformHandle)
 					{
 						UE_LOG(LogEOSVoiceChat, Warning, TEXT("FEOSVoiceChat::Initialize CreatePlatform failed"));
-						InitSession = FInitSession();
+						InitSession.Reset();
 						Result = FVoiceChatResult(EVoiceChatResult::ImplementationError);
 					}
 				}
@@ -185,7 +185,7 @@ void FEOSVoiceChat::Initialize(const FOnVoiceChatInitializeCompleteDelegate& Ini
 					else
 					{
 						UE_LOG(LogEOSVoiceChat, Warning, TEXT("FEOSVoiceChat::Initialize failed to get RTC interface handle"));
-						InitSession = FInitSession();
+						InitSession.Reset();
 						Result = FVoiceChatResult(EVoiceChatResult::ImplementationError);
 					}
 				}
@@ -193,7 +193,7 @@ void FEOSVoiceChat::Initialize(const FOnVoiceChatInitializeCompleteDelegate& Ini
 			else
 			{
 				UE_LOG(LogEOSVoiceChat, Warning, TEXT("FEOSVoiceChat::Initialize Initialize failed"));
-				InitSession = FInitSession();
+				InitSession.Reset();
 				Result = FVoiceChatResult(EVoiceChatResult::ImplementationError);
 			}
 		}
@@ -244,7 +244,7 @@ void FEOSVoiceChat::Uninitialize(const FOnVoiceChatUninitializeCompleteDelegate&
 			UnbindInitCallbacks();
 
 			const TArray<FOnVoiceChatUninitializeCompleteDelegate> UninitializeCompleteDelegates = MoveTemp(InitSession.UninitializeCompleteDelegates);
-			InitSession = FInitSession();
+			InitSession.Reset();
 			for (const FOnVoiceChatUninitializeCompleteDelegate& UninitializeCompleteDelegate : UninitializeCompleteDelegates)
 			{
 				UninitializeCompleteDelegate.ExecuteIfBound(FVoiceChatResult::CreateSuccess());
@@ -1021,6 +1021,21 @@ FEOSVoiceChat::FInitSession::FInitSession()
 {
 }
 
+void FEOSVoiceChat::FInitSession::Reset()
+{
+	State = EInitializationState::Uninitialized;
+
+	UninitializeCompleteDelegates = TArray<FOnVoiceChatUninitializeCompleteDelegate>{};
+
+	EosPlatformHandle = nullptr;
+	EosRtcInterface = nullptr;
+	EosLobbyInterface = nullptr;
+
+	OnAudioDevicesChangedNotificationId = EOS_INVALID_NOTIFICATIONID;
+
+	EosAudioDevicePool = MakeShared<FEOSAudioDevicePool>(EosRtcInterface);
+}
+
 IEOSPlatformHandlePtr FEOSVoiceChat::EOSPlatformCreate(EOS_Platform_Options& PlatformOptions)
 {
 	return SDKManager.CreatePlatform(PlatformOptions);
diff --git a/Engine/Plugins/Online/VoiceChat/EOSVoiceChat/Source/EOSVoiceChat/Public/EOSVoiceChat.h b/Engine/Plugins/Online/VoiceChat/EOSVoiceChat/Source/EOSVoiceChat/Public/EOSVoiceChat.h
--- a/Engine/Plugins/Online/VoiceChat/EOSVoiceChat/Source/EOSVoiceChat/Public/EOSVoiceChat.h
+++ b/Engine/Plugins/Online/VoiceChat/EOSVoiceChat/Source/EOSVoiceChat/Public/EOSVoiceChat.h
@@ -167,6 +167,12 @@ protected:
 	{
 		FInitSession();
 
+		// `FEOSAudioDevicePool` holds a ref to an object field itself which can't be copied or moved
+		UE_NONCOPYABLE(FInitSession)
+
+		// reverts fields to their default state
+		void Reset();
+
 		EInitializationState State = EInitializationState::Uninitialized;
 
 		TArray<FOnVoiceChatUninitializeCompleteDelegate> UninitializeCompleteDelegates;
