- Move EOSPlatformHandle out of FInitSession so it persists for the lifetime of FEOSVoiceChat.
- Avoids GT hitches when VC is enabled/disabled during gameplay.

diff --git a/Engine/Plugins/Online/VoiceChat/EOSVoiceChat/Source/EOSVoiceChat/Private/EOSVoiceChat.cpp b/Engine/Plugins/Online/VoiceChat/EOSVoiceChat/Source/EOSVoiceChat/Private/EOSVoiceChat.cpp
--- a/Engine/Plugins/Online/VoiceChat/EOSVoiceChat/Source/EOSVoiceChat/Private/EOSVoiceChat.cpp
+++ b/Engine/Plugins/Online/VoiceChat/EOSVoiceChat/Source/EOSVoiceChat/Private/EOSVoiceChat.cpp
@@ -50,9 +50,9 @@ int64 FEOSVoiceChat::StaticInstanceIdCount = 0;
 
 #define EOS_VOICE_TODO 0
 
-FEOSVoiceChat::FEOSVoiceChat(IEOSSDKManager& InSDKManager, const IEOSPlatformHandlePtr& PlatformHandle)
+FEOSVoiceChat::FEOSVoiceChat(IEOSSDKManager& InSDKManager, const IEOSPlatformHandlePtr& InPlatformHandle)
 	: SDKManager(InSDKManager)
-	, ExternalPlatformHandle(PlatformHandle)
+	, EosPlatformHandle(InPlatformHandle)
 {
 }
 
@@ -81,7 +81,7 @@ bool FEOSVoiceChat::Uninitialize()
 
 	while (!bIsDone)
 	{
-		InitSession.EosPlatformHandle->Tick();
+		EosPlatformHandle->Tick();
 	}
 	
 	return !IsInitialized();
@@ -101,14 +101,10 @@ void FEOSVoiceChat::Initialize(const FOnVoiceChatInitializeCompleteDelegate& Ini
 		{
 			InitSession.State = EInitializationState::Initializing;
 
-			EOS_EResult EosResult = SDKManager.Initialize();
-			if (EosResult == EOS_EResult::EOS_Success)
+			if (!EosPlatformHandle)
 			{
-				if(ExternalPlatformHandle)
-				{
-					InitSession.EosPlatformHandle = ExternalPlatformHandle;
-				}
-				else
+				EOS_EResult EosResult = SDKManager.Initialize();
+				if (EosResult == EOS_EResult::EOS_Success)
 				{
 					FString ConfigProductId;
 					FString ConfigSandboxId;
@@ -162,45 +158,48 @@ void FEOSVoiceChat::Initialize(const FOnVoiceChatInitializeCompleteDelegate& Ini
 					static_assert(EOS_PLATFORM_RTCOPTIONS_API_LATEST == 1, "EOS_Platform_RTCOptions updated, check new fields");
 					PlatformOptions.RTCOptions = &PlatformRTCOptions;
 
-					InitSession.EosPlatformHandle = EOSPlatformCreate(PlatformOptions);
-					if (!InitSession.EosPlatformHandle)
+					EosPlatformHandle = EOSPlatformCreate(PlatformOptions);
+					if (!EosPlatformHandle)
 					{
 						UE_LOG(LogEOSVoiceChat, Warning, TEXT("FEOSVoiceChat::Initialize CreatePlatform failed"));
-						InitSession.Reset();
 						Result = FVoiceChatResult(EVoiceChatResult::ImplementationError);
 					}
 				}
-
-				if (Result.IsSuccess())
+				else
 				{
-					InitSession.EosRtcInterface = EOS_Platform_GetRTCInterface(*InitSession.EosPlatformHandle);
-					InitSession.EosLobbyInterface = EOS_Platform_GetLobbyInterface(*InitSession.EosPlatformHandle);
-					if (InitSession.EosRtcInterface)
-					{
-						BindInitCallbacks();
-						InitSession.State = EInitializationState::Initialized;
-						PostInitialize();
-						Result = FVoiceChatResult::CreateSuccess();
-					}
-					else
-					{
-						UE_LOG(LogEOSVoiceChat, Warning, TEXT("FEOSVoiceChat::Initialize failed to get RTC interface handle"));
-						InitSession.Reset();
-						Result = FVoiceChatResult(EVoiceChatResult::ImplementationError);
-					}
+					UE_LOG(LogEOSVoiceChat, Warning, TEXT("FEOSVoiceChat::Initialize SDKManager.Initialize failed"));
+					Result = FVoiceChatResult(EVoiceChatResult::ImplementationError);
 				}
 			}
-			else
+
+			if (Result.IsSuccess())
 			{
-				UE_LOG(LogEOSVoiceChat, Warning, TEXT("FEOSVoiceChat::Initialize Initialize failed"));
-				InitSession.Reset();
-				Result = FVoiceChatResult(EVoiceChatResult::ImplementationError);
+				InitSession.EosRtcInterface = EOS_Platform_GetRTCInterface(*EosPlatformHandle);
+				InitSession.EosLobbyInterface = EOS_Platform_GetLobbyInterface(*EosPlatformHandle);
+				if (InitSession.EosRtcInterface && InitSession.EosLobbyInterface)
+				{
+					BindInitCallbacks();
+					InitSession.State = EInitializationState::Initialized;
+					PostInitialize();
+					Result = FVoiceChatResult::CreateSuccess();
+				}
+				else
+				{
+					UE_LOG(LogEOSVoiceChat, Warning, TEXT("FEOSVoiceChat::Initialize failed to get interface handles"));
+					Result = FVoiceChatResult(EVoiceChatResult::ImplementationError);
+				}
 			}
 		}
 		else
 		{
 			Result = VoiceChat::Errors::NotEnabled();
 		}
+
+		if (!Result.IsSuccess())
+		{
+			InitSession.Reset();
+		}
+
 		break;
 	}
 	case EInitializationState::Uninitializing:
@@ -785,7 +784,7 @@ void FEOSVoiceChat::BindInitCallbacks()
 	EOS_RTCAudio_AddNotifyAudioDevicesChangedOptions AudioDevicesChangedOptions = {};
 	AudioDevicesChangedOptions.ApiVersion = EOS_RTCAUDIO_ADDNOTIFYAUDIODEVICESCHANGED_API_LATEST;
 	static_assert(EOS_RTCAUDIO_ADDNOTIFYAUDIODEVICESCHANGED_API_LATEST == 1, "EOS_RTC_AddNotifyAudioDevicesChangedOptions updated, check new fields");
-	InitSession.OnAudioDevicesChangedNotificationId = EOS_RTCAudio_AddNotifyAudioDevicesChanged(EOS_RTC_GetAudioInterface(InitSession.EosRtcInterface), &AudioDevicesChangedOptions, this, &FEOSVoiceChat::OnAudioDevicesChangedStatic);
+	InitSession.OnAudioDevicesChangedNotificationId = EOS_RTCAudio_AddNotifyAudioDevicesChanged(EOS_RTC_GetAudioInterface(GetRtcInterface()), &AudioDevicesChangedOptions, this, &FEOSVoiceChat::OnAudioDevicesChangedStatic);
 	if (InitSession.OnAudioDevicesChangedNotificationId == EOS_INVALID_NOTIFICATIONID)
 	{
 		UE_LOG(LogEOSVoiceChat, Warning, TEXT("BindInitCallbacks EOS_RTC_AddNotifyAudioDevicesChanged failed"));
@@ -798,7 +797,7 @@ void FEOSVoiceChat::UnbindInitCallbacks()
 {
 	if (InitSession.OnAudioDevicesChangedNotificationId != EOS_INVALID_NOTIFICATIONID)
 	{
-		EOS_RTCAudio_RemoveNotifyAudioDevicesChanged(EOS_RTC_GetAudioInterface(InitSession.EosRtcInterface), InitSession.OnAudioDevicesChangedNotificationId);
+		EOS_RTCAudio_RemoveNotifyAudioDevicesChanged(EOS_RTC_GetAudioInterface(GetRtcInterface()), InitSession.OnAudioDevicesChangedNotificationId);
 		InitSession.OnAudioDevicesChangedNotificationId = EOS_INVALID_NOTIFICATIONID;
 	}
 }
@@ -1027,7 +1026,6 @@ void FEOSVoiceChat::FInitSession::Reset()
 
 	UninitializeCompleteDelegates = TArray<FOnVoiceChatUninitializeCompleteDelegate>{};
 
-	EosPlatformHandle = nullptr;
 	EosRtcInterface = nullptr;
 	EosLobbyInterface = nullptr;
 
diff --git a/Engine/Plugins/Online/VoiceChat/EOSVoiceChat/Source/EOSVoiceChat/Public/EOSVoiceChat.h b/Engine/Plugins/Online/VoiceChat/EOSVoiceChat/Source/EOSVoiceChat/Public/EOSVoiceChat.h
--- a/Engine/Plugins/Online/VoiceChat/EOSVoiceChat/Source/EOSVoiceChat/Public/EOSVoiceChat.h
+++ b/Engine/Plugins/Online/VoiceChat/EOSVoiceChat/Source/EOSVoiceChat/Public/EOSVoiceChat.h
@@ -177,7 +177,6 @@ protected:
 
 		TArray<FOnVoiceChatUninitializeCompleteDelegate> UninitializeCompleteDelegates;
 
-		IEOSPlatformHandlePtr EosPlatformHandle = nullptr;
 		EOS_HRTC EosRtcInterface = nullptr;
 		EOS_HLobby EosLobbyInterface = nullptr;
 
@@ -187,7 +186,6 @@ protected:
 	};
 	FInitSession InitSession;
 
-	virtual IEOSPlatformHandlePtr EOSPlatformCreate(EOS_Platform_Options& PlatformOptions);
 	void BindInitCallbacks();
 	void UnbindInitCallbacks();
 	static void EOS_CALL OnAudioDevicesChangedStatic(const EOS_RTCAudio_AudioDevicesChangedCallbackInfo* CallbackInfo);
@@ -222,7 +220,8 @@ protected:
 	TArray<IVoiceChatUser*> UsersCreatedByConsoleCommand;
 
 	IEOSSDKManager& SDKManager;
-	IEOSPlatformHandlePtr ExternalPlatformHandle;
+	IEOSPlatformHandlePtr EosPlatformHandle = nullptr;
+	virtual IEOSPlatformHandlePtr EOSPlatformCreate(EOS_Platform_Options& PlatformOptions);
 
 	static int64 StaticInstanceIdCount;
 	int64 InstanceId = StaticInstanceIdCount++;
