- Fix horde failures

diff --git a/Engine/Plugins/Online/VoiceChat/EOSVoiceChat/Source/EOSVoiceChat/Public/EOSVoiceChat.h b/Engine/Plugins/Online/VoiceChat/EOSVoiceChat/Source/EOSVoiceChat/Public/EOSVoiceChat.h
--- a/Engine/Plugins/Online/VoiceChat/EOSVoiceChat/Source/EOSVoiceChat/Public/EOSVoiceChat.h
+++ b/Engine/Plugins/Online/VoiceChat/EOSVoiceChat/Source/EOSVoiceChat/Public/EOSVoiceChat.h
@@ -143,7 +143,7 @@ public:
 	virtual FString InsecureGetJoinToken(const FString& ChannelName, EVoiceChatChannelType ChannelType, TOptional<FVoiceChatChannel3dProperties> Channel3dProperties = TOptional<FVoiceChatChannel3dProperties>()) override;
 	// ~End IVoiceChatUser Interface
 
-	IEOSPlatformHandlePtr GetPlatformHandle() const { return InitSession.EosPlatformHandle; }
+	IEOSPlatformHandlePtr GetPlatformHandle() const { return EosPlatformHandle; }
 	EOS_HRTC GetRtcInterface() const { return InitSession.EosRtcInterface; }
 	EOS_HLobby GetLobbyInterface() const { return InitSession.EosLobbyInterface; }
 
