- Add SetPlayerVolume support
- Update documentation of Volume fields in VoiceChat.h, it is now explicitly from 0-2, with 1 being no difference from source volume.

diff --git a/Engine/Plugins/Online/VoiceChat/EOSVoiceChat/Source/EOSVoiceChat/Private/EOSVoiceChatUser.cpp b/Engine/Plugins/Online/VoiceChat/EOSVoiceChat/Source/EOSVoiceChat/Private/EOSVoiceChatUser.cpp
--- a/Engine/Plugins/Online/VoiceChat/EOSVoiceChat/Source/EOSVoiceChat/Private/EOSVoiceChatUser.cpp
+++ b/Engine/Plugins/Online/VoiceChat/EOSVoiceChat/Source/EOSVoiceChat/Private/EOSVoiceChatUser.cpp
@@ -152,7 +152,7 @@ void FEOSVoiceChatUser::SetAudioInputVolume(float InVolume)
 {
 	EOSVOICECHATUSER_LOG(Verbose, TEXT("SetAudioInputVolume %f"), InVolume);
 
-	InVolume = FMath::Clamp(InVolume, 0.0f, 1.0f);
+	InVolume = FMath::Clamp(InVolume, 0.0f, 2.0f);
 
 	if (AudioInputOptions.Volume != InVolume)
 	{
@@ -166,7 +166,7 @@ void FEOSVoiceChatUser::SetAudioOutputVolume(float InVolume)
 {
 	EOSVOICECHATUSER_LOG(Verbose, TEXT("SetAudioOutputVolume %f"), InVolume);
 
-	InVolume = FMath::Clamp(InVolume, 0.0f, 1.0f);
+	InVolume = FMath::Clamp(InVolume, 0.0f, 2.0f);
 
 	if (AudioOutputOptions.Volume != InVolume)
 	{
@@ -370,7 +370,7 @@ void FEOSVoiceChatUser::Login(FPlatformUserId PlatformId, const FString& PlayerN
 		return;
 	}
 
-	EOS_ProductUserId LocalUserProductUserId = EOS_ProductUserId_FromString(TCHAR_TO_UTF8(*PlayerName));
+	const EOS_ProductUserId LocalUserProductUserId = EOSProductUserIdFromString(*PlayerName);
 	if (EOS_ProductUserId_IsValid(LocalUserProductUserId))
 	{
 		LoginSession.PlatformId = PlatformId;
@@ -569,7 +569,7 @@ void FEOSVoiceChatUser::JoinChannel(const FString& ChannelName, const FString& C
 			const FTCHARToUTF8 Utf8RoomName(*ChannelName);
 			const FTCHARToUTF8 Utf8ClientBaseUrl(*ChannelCredentials.ClientBaseUrl);
 			const FTCHARToUTF8 Utf8ParticipantToken(*ChannelCredentials.ParticipantToken);
-			const FTCHARToUTF8 Utf8ParticipantId(*ChannelSession.PlayerName);
+			const EOS_ProductUserId ParticipantId = EOSProductUserIdFromString(*ChannelSession.PlayerName);
 
 			// Attempt to join the channel
 			{
@@ -579,7 +579,7 @@ void FEOSVoiceChatUser::JoinChannel(const FString& ChannelName, const FString& C
 				JoinOptions.LocalUserId = LoginSession.LocalUserProductUserId;
 				JoinOptions.RoomName = Utf8RoomName.Get();
 				JoinOptions.ClientBaseUrl = Utf8ClientBaseUrl.Get();
-				JoinOptions.ParticipantId = EOS_ProductUserId_FromString(Utf8ParticipantId.Get());
+				JoinOptions.ParticipantId = ParticipantId;
 				JoinOptions.ParticipantToken = Utf8ParticipantToken.Get();
 #if !UE_BUILD_SHIPPING
 				const bool bEnableEcho = ChannelSession.ChannelType == EVoiceChatChannelType::Echo || CVarChannelEchoEnabled.GetValueOnAnyThread();
@@ -706,14 +706,7 @@ void FEOSVoiceChatUser::SetPlayerMuted(const FString& PlayerName, bool bAudioMut
 	{
 		GlobalParticipant.bAudioMuted = bAudioMuted;
 
-		for (TPair<FString, FChannelSession>& Pair : LoginSession.ChannelSessions)
-		{
-			FChannelSession& ChannelSession = Pair.Value;
-			if (FChannelParticipant* ChannelParticipant = ChannelSession.Participants.Find(PlayerName))
-			{
-				ApplyPlayerReceivingOptions(GlobalParticipant, ChannelSession, *ChannelParticipant);
-			}
-		}
+		ApplyPlayerReceivingOptions(PlayerName);
 	}
 }
 
@@ -761,15 +754,15 @@ void FEOSVoiceChatUser::SetPlayerVolume(const FString& PlayerName, float Volume)
 {
 	EOSVOICECHATUSER_LOG(Log, TEXT("SetPlayerVolume PlayerName=[%s] Volume=[%s]"), *PlayerName, *LexToString(Volume));
 
-	// Can't actually set the volume so we just mute/unmute
-	// TODO from an IVoiceChat perspective, Volume/mute should be independent settings.
-	const bool bAudioMuted = Volume <= SMALL_NUMBER;
-	SetPlayerMuted(PlayerName, bAudioMuted);
+	FGlobalParticipant& GlobalParticipant = GetGlobalParticipant(PlayerName);
+	GlobalParticipant.Volume = FMath::Clamp(Volume, 0.0f, 2.0f);
+
+	ApplyPlayerReceivingOptions(PlayerName);
 }
 
 float FEOSVoiceChatUser::GetPlayerVolume(const FString& PlayerName) const
 {
-	return GetGlobalParticipant(PlayerName).bAudioMuted ? 0.0f : 1.0f;
+	return GetGlobalParticipant(PlayerName).Volume;
 }
 
 void FEOSVoiceChatUser::TransmitToAllChannels()
@@ -1196,14 +1189,14 @@ void FEOSVoiceChatUser::ApplyPlayerBlock(const FGlobalParticipant& GlobalPartici
 	EOSVOICECHATUSER_LOG(Verbose, TEXT("ApplyPlayerBlock ChannelName=[%s] PlayerName=[%s] bBlocked=%s"), *ChannelSession.ChannelName, *GlobalParticipant.PlayerName, *LexToString(GlobalParticipant.bBlocked));
 
 	const FTCHARToUTF8 Utf8RoomName(*ChannelSession.ChannelName);
-	const FTCHARToUTF8 Utf8ParticipantId(*GlobalParticipant.PlayerName);
+	const EOS_ProductUserId ParticipantId = EOSProductUserIdFromString(*GlobalParticipant.PlayerName);
 
 	EOS_RTC_BlockParticipantOptions Options = {};
 	Options.ApiVersion = EOS_RTC_BLOCKPARTICIPANT_API_LATEST;
 	static_assert(EOS_RTC_BLOCKPARTICIPANT_API_LATEST == 1, "EOS_RTC_BlockParticipantOptions updated, check new fields");
 	Options.LocalUserId = LoginSession.LocalUserProductUserId;
 	Options.RoomName = Utf8RoomName.Get();
-	Options.ParticipantId = EOS_ProductUserId_FromString(Utf8ParticipantId.Get());
+	Options.ParticipantId = ParticipantId;
 	Options.bBlocked = GlobalParticipant.bBlocked;
 
 	CSV_SCOPED_TIMING_STAT_EXCLUSIVE(EOSVoiceChat);
@@ -1228,6 +1221,22 @@ void FEOSVoiceChatUser::ApplyReceivingOptions(const FChannelSession& ChannelSess
 	EOS_RTCAudio_UpdateReceiving(EOS_RTC_GetAudioInterface(GetRtcInterface()), &UpdateReceivingOptions, this, &FEOSVoiceChatUser::OnUpdateReceivingAudioStatic);
 }
 
+void FEOSVoiceChatUser::ApplyPlayerReceivingOptions(const FString& PlayerName)
+{
+	check(IsLoggedIn());
+
+	FGlobalParticipant& GlobalParticipant = GetGlobalParticipant(PlayerName);
+
+	for (TPair<FString, FChannelSession>& Pair : LoginSession.ChannelSessions)
+	{
+		FChannelSession& ChannelSession = Pair.Value;
+		if (FChannelParticipant* ChannelParticipant = ChannelSession.Participants.Find(PlayerName))
+		{
+			ApplyPlayerReceivingOptions(GlobalParticipant, ChannelSession, *ChannelParticipant);
+		}
+	}
+}
+
 void FEOSVoiceChatUser::ApplyPlayerReceivingOptions(const FGlobalParticipant& GlobalParticipant, const FChannelSession& ChannelSession, FChannelParticipant& ChannelParticipant)
 {
 	// TODO ChannelParticipant should be const, and we should update the bAudioMuted (the *actual* state, not the *desired* state) in the OnUpdateReceiving callback.
@@ -1237,19 +1246,35 @@ void FEOSVoiceChatUser::ApplyPlayerReceivingOptions(const FGlobalParticipant& Gl
 	EOSVOICECHATUSER_LOG(Verbose, TEXT("ApplyPlayerReceivingOptions ChannelName=[%s] PlayerName=[%s] bAudioMuted=[%s]"), *ChannelSession.ChannelName, *GlobalParticipant.PlayerName, *LexToString(GlobalParticipant.bAudioMuted));
 
 	const FTCHARToUTF8 Utf8RoomName(*ChannelSession.ChannelName);
-	const FTCHARToUTF8 Utf8ParticipantId(*GlobalParticipant.PlayerName);
+	const EOS_ProductUserId ParticipantId = EOSProductUserIdFromString(*GlobalParticipant.PlayerName);
 
-	EOS_RTCAudio_UpdateReceivingOptions UpdateReceivingOptions = {};
-	UpdateReceivingOptions.ApiVersion = EOS_RTCAUDIO_UPDATERECEIVING_API_LATEST;
-	static_assert(EOS_RTCAUDIO_UPDATERECEIVING_API_LATEST == 1, "EOS_RTCAudio_UpdateReceivingOptions updated, check new fields");
-	UpdateReceivingOptions.LocalUserId = LoginSession.LocalUserProductUserId;
-	UpdateReceivingOptions.RoomName = Utf8RoomName.Get();
-	UpdateReceivingOptions.ParticipantId = EOS_ProductUserId_FromString(Utf8ParticipantId.Get());
-	UpdateReceivingOptions.bAudioEnabled = ChannelParticipant.bAudioDisabled ? EOS_FALSE : EOS_TRUE;
+	{
+		EOS_RTCAudio_UpdateReceivingOptions UpdateReceivingOptions = {};
+		UpdateReceivingOptions.ApiVersion = EOS_RTCAUDIO_UPDATERECEIVING_API_LATEST;
+		static_assert(EOS_RTCAUDIO_UPDATERECEIVING_API_LATEST == 1, "EOS_RTCAudio_UpdateReceivingOptions updated, check new fields");
+		UpdateReceivingOptions.LocalUserId = LoginSession.LocalUserProductUserId;
+		UpdateReceivingOptions.RoomName = Utf8RoomName.Get();
+		UpdateReceivingOptions.ParticipantId = ParticipantId;
+		UpdateReceivingOptions.bAudioEnabled = ChannelParticipant.bAudioDisabled ? EOS_FALSE : EOS_TRUE;
 
-	CSV_SCOPED_TIMING_STAT_EXCLUSIVE(EOSVoiceChat);
-	QUICK_SCOPE_CYCLE_COUNTER(EOS_RTCAudio_UpdateReceiving);
-	EOS_RTCAudio_UpdateReceiving(EOS_RTC_GetAudioInterface(GetRtcInterface()), &UpdateReceivingOptions, this, &FEOSVoiceChatUser::OnUpdateReceivingAudioStatic);
+		CSV_SCOPED_TIMING_STAT_EXCLUSIVE(EOSVoiceChat);
+		QUICK_SCOPE_CYCLE_COUNTER(EOS_RTCAudio_UpdateReceiving);
+		EOS_RTCAudio_UpdateReceiving(EOS_RTC_GetAudioInterface(GetRtcInterface()), &UpdateReceivingOptions, this, &FEOSVoiceChatUser::OnUpdateReceivingAudioStatic);
+	}
+
+	{
+		EOS_RTCAudio_UpdateParticipantVolumeOptions UpdateParticipantVolumeOptions = {};
+		UpdateParticipantVolumeOptions.ApiVersion = EOS_RTCAUDIO_UPDATEPARTICIPANTVOLUME_API_LATEST;
+		static_assert(EOS_RTCAUDIO_UPDATEPARTICIPANTVOLUME_API_LATEST == 1, "EOS_RTCAudio_UpdateParticipantVolumeOptions updated, check new fields");
+		UpdateParticipantVolumeOptions.LocalUserId = LoginSession.LocalUserProductUserId;
+		UpdateParticipantVolumeOptions.RoomName = Utf8RoomName.Get();
+		UpdateParticipantVolumeOptions.ParticipantId = ParticipantId;
+		UpdateParticipantVolumeOptions.Volume = GlobalParticipant.Volume * 50.f;
+
+		CSV_SCOPED_TIMING_STAT_EXCLUSIVE(EOSVoiceChat);
+		QUICK_SCOPE_CYCLE_COUNTER(EOS_RTCAudio_UpdateParticipantVolume);
+		EOS_RTCAudio_UpdateParticipantVolume(EOS_RTC_GetAudioInterface(GetRtcInterface()), &UpdateParticipantVolumeOptions, this, &FEOSVoiceChatUser::OnUpdateParticipantVolumeStatic);
+	}
 }
 
 void FEOSVoiceChatUser::ApplySendingOptions()
@@ -1776,6 +1801,56 @@ void FEOSVoiceChatUser::OnBlockParticipant(const EOS_RTC_BlockParticipantCallbac
 	}
 }
 
+void EOS_CALL FEOSVoiceChatUser::OnUpdateParticipantVolumeStatic(const EOS_RTCAudio_UpdateParticipantVolumeCallbackInfo* CallbackInfo)
+{
+	CSV_SCOPED_TIMING_STAT_EXCLUSIVE(EOSVoiceChat);
+	QUICK_SCOPE_CYCLE_COUNTER(STAT_FEOSVoiceChat_OnUpdateParticipantVolumeStatic);
+
+	if (CallbackInfo)
+	{
+		if (FEOSVoiceChatUser* EosVoiceChatPtr = static_cast<FEOSVoiceChatUser*>(CallbackInfo->ClientData))
+		{
+			EosVoiceChatPtr->OnUpdateParticipantVolume(CallbackInfo);
+		}
+		else
+		{
+			UE_LOG(LogEOSVoiceChat, Warning, TEXT("OnUpdateParticipantVolumeStatic Error EosVoiceChatPtr=nullptr"));
+		}
+	}
+	else
+	{
+		UE_LOG(LogEOSVoiceChat, Warning, TEXT("OnUpdateParticipantVolumeStatic Error CallbackInfo=nullptr"));
+	}
+}
+
+void FEOSVoiceChatUser::OnUpdateParticipantVolume(const EOS_RTCAudio_UpdateParticipantVolumeCallbackInfo* CallbackInfo)
+{
+	check(IsInitialized());
+
+	const FString ChannelName = UTF8_TO_TCHAR(CallbackInfo->RoomName);
+	const FString PlayerName = LexToString(CallbackInfo->ParticipantId);
+	const EOS_EResult& ResultCode = CallbackInfo->ResultCode;
+
+	if (ResultCode == EOS_EResult::EOS_Success)
+	{
+		EOSVOICECHATUSER_LOG(Log, TEXT("OnUpdateParticipantVolume ChannelName=[%s] PlayerName=[%s] Success"), *ChannelName, *PlayerName);
+		FChannelSession& ChannelSession = GetChannelSession(ChannelName);
+		if (FChannelParticipant* ChannelParticipant = ChannelSession.Participants.Find(PlayerName))
+		{
+			const FGlobalParticipant& GlobalParticipant = GetGlobalParticipant(PlayerName);
+			OnVoiceChatPlayerVolumeUpdatedDelegate.Broadcast(ChannelSession.ChannelName, ChannelParticipant->PlayerName, GlobalParticipant.Volume);
+		}
+		else
+		{
+			EOSVOICECHATUSER_LOG(Warning, TEXT("OnUpdateParticipantVolume ChannelName=[%s] PlayerName=[%s] failed to find player in channel."), *ChannelName, *PlayerName);
+		}
+	}
+	else
+	{
+		EOSVOICECHATUSER_LOG(Warning, TEXT("OnUpdateParticipantVolume ChannelName=[%s] PlayerName=[%s] failed error=[%s]"), *ChannelName, *PlayerName, *LexToString(ResultCode));
+	}
+}
+
 void EOS_CALL FEOSVoiceChatUser::OnUpdateReceivingAudioStatic(const EOS_RTCAudio_UpdateReceivingCallbackInfo* CallbackInfo)
 {
 	if (CallbackInfo)
diff --git a/Engine/Plugins/Online/VoiceChat/EOSVoiceChat/Source/EOSVoiceChat/Public/EOSVoiceChatUser.h b/Engine/Plugins/Online/VoiceChat/EOSVoiceChat/Source/EOSVoiceChat/Public/EOSVoiceChatUser.h
--- a/Engine/Plugins/Online/VoiceChat/EOSVoiceChat/Source/EOSVoiceChat/Public/EOSVoiceChatUser.h
+++ b/Engine/Plugins/Online/VoiceChat/EOSVoiceChat/Source/EOSVoiceChat/Public/EOSVoiceChatUser.h
@@ -103,6 +103,8 @@ protected:
 		bool bBlocked = false;
 		// Desired mute state
 		bool bAudioMuted = false;
+		// Desired Volume
+		float Volume = 1.f;
 	};
 
 	/**
@@ -301,6 +303,7 @@ protected:
 	void ApplyAudioOutputOptions();
 	void ApplyPlayerBlock(const FGlobalParticipant& GlobalParticipant, const FChannelSession& ChannelSession, FChannelParticipant& ChannelParticipant);
 	void ApplyReceivingOptions(const FChannelSession& ChannelSession);
+	void ApplyPlayerReceivingOptions(const FString& PlayerName);
 	void ApplyPlayerReceivingOptions(const FGlobalParticipant& GlobalParticipant, const FChannelSession& ChannelSession, FChannelParticipant& ChannelParticipant);
 	void ApplySendingOptions();
 	void ApplySendingOptions(FChannelSession& ChannelSession);
@@ -326,6 +329,8 @@ protected:
 	void OnLeaveRoom(const EOS_RTC_LeaveRoomCallbackInfo* CallbackInfo);
 	static void EOS_CALL OnBlockParticipantStatic(const EOS_RTC_BlockParticipantCallbackInfo* CallbackInfo);
 	void OnBlockParticipant(const EOS_RTC_BlockParticipantCallbackInfo* CallbackInfo);
+	static void EOS_CALL OnUpdateParticipantVolumeStatic(const EOS_RTCAudio_UpdateParticipantVolumeCallbackInfo* CallbackInfo);
+	void OnUpdateParticipantVolume(const EOS_RTCAudio_UpdateParticipantVolumeCallbackInfo* CallbackInfo);
 	static void EOS_CALL OnUpdateReceivingAudioStatic(const EOS_RTCAudio_UpdateReceivingCallbackInfo* CallbackInfo);
 	void OnUpdateReceivingAudio(const EOS_RTCAudio_UpdateReceivingCallbackInfo* CallbackInfo);
 	static void EOS_CALL OnUpdateSendingAudioStatic(const EOS_RTCAudio_UpdateSendingCallbackInfo* CallbackInfo);
diff --git a/Engine/Plugins/Online/VoiceChat/VoiceChat/Source/Public/VoiceChat.h b/Engine/Plugins/Online/VoiceChat/VoiceChat/Source/Public/VoiceChat.h
--- a/Engine/Plugins/Online/VoiceChat/VoiceChat/Source/Public/VoiceChat.h
+++ b/Engine/Plugins/Online/VoiceChat/VoiceChat/Source/Public/VoiceChat.h
@@ -140,28 +140,28 @@ public:
 	/**
 	 * Set record volume
 	 *
-	 * @param Volume Value between 0.0 and 1.0
+	 * @param Volume Value between 0.0 and 2.0. 1.0 leaves the source volume unchanged.
 	 */
 	virtual void SetAudioInputVolume(float Volume) = 0;
 
 	/**
 	 * Set playback volume
 	 *
-	 * @param Volume Value between 0.0 and 1.0
+	 * @param Volume Value between 0.0 and 2.0. 1.0 leaves the source volume unchanged.
 	 */
 	virtual void SetAudioOutputVolume(float Volume) = 0;
 
 	/**
 	 * Get volume of audio input device
 	 *
-	 * @return Audio input device volume. Will be between 0.0 and 1.0
+	 * @return Audio input device volume between 0.0 and 2.0. 1.0 leaves the source volume unchanged.
 	 */
 	virtual float GetAudioInputVolume() const = 0;
 
 	/**
 	 * Get volume of audio output device
 	 *
-	 * @return Audio output device volume. Will be between 0.0 and 1.0
+	 * @return Audio output device volume between 0.0 and 2.0. 1.0 leaves the source volume unchanged.
 	 */
 	virtual float GetAudioOutputVolume() const = 0;
 
@@ -551,13 +551,13 @@ public:
 	 * Set player volume
 	 *
 	 * @param PlayerName Player to adjust the volume for
-	 * @param Volume Value between 0.0 and 1.0
+	 * @param Volume Value between 0.0 and 2.0. 1.0 leaves the source volume unchanged.
 	 */
 	virtual void SetPlayerVolume(const FString& PlayerName, float Volume) = 0;
 
 	/**
 	 * @param PlayerName Player to get the volume of
-	 * @return volume for player
+	 * @return Volume Value between 0.0 and 2.0. 1.0 leaves the source volume unchanged.
 	 */
 	virtual float GetPlayerVolume(const FString& PlayerName) const = 0;
 
