- Re-add indirection through TGlobalCallback for RTC AudioBeforeSend callback, to fix crash when it is called back after EOSVoiceChat has been torn down. The callback is added to a list of callbacks on EOSShared which is cleared when EOS_Shutdown is called.

diff --git a/Engine/Plugins/Online/EOSShared/Source/EOSShared/Private/EOSSDKManager.cpp b/Engine/Plugins/Online/EOSShared/Source/EOSShared/Private/EOSSDKManager.cpp
--- a/Engine/Plugins/Online/EOSShared/Source/EOSShared/Private/EOSSDKManager.cpp
+++ b/Engine/Plugins/Online/EOSShared/Source/EOSShared/Private/EOSSDKManager.cpp
@@ -19,6 +19,7 @@
 #include "CoreGlobals.h"
 
 #include "EOSShared.h"
+#include "EOSSharedTypes.h"
 
 #include "eos_auth.h"
 #include "eos_connect.h"
@@ -457,6 +458,8 @@ void FEOSSDKManager::Shutdown()
 		const EOS_EResult Result = EOS_Shutdown();
 		UE_LOG(LogEOSSDK, Log, TEXT("FEOSSDKManager::Shutdown EOS_Shutdown Result=[%s]"), *LexToString(Result));
 
+		CallbackObjects.Empty();
+
 		bInitialized = false;
 	}
 }
@@ -807,6 +810,11 @@ void FEOSSDKManager::LogConnectInfo(const EOS_HPlatform Platform, const EOS_Prod
 	}
 }
 
+void FEOSSDKManager::AddCallbackObject(TUniquePtr<FCallbackBase> CallbackObj)
+{
+	CallbackObjects.Emplace(MoveTemp(CallbackObj));
+}
+
 FEOSPlatformHandle::~FEOSPlatformHandle()
 {
 	Manager.ReleasePlatform(PlatformHandle);
diff --git a/Engine/Plugins/Online/EOSShared/Source/EOSShared/Private/EOSSDKManager.h b/Engine/Plugins/Online/EOSShared/Source/EOSShared/Private/EOSSDKManager.h
--- a/Engine/Plugins/Online/EOSShared/Source/EOSShared/Private/EOSSDKManager.h
+++ b/Engine/Plugins/Online/EOSShared/Source/EOSShared/Private/EOSSDKManager.h
@@ -47,6 +47,8 @@ public:
 	virtual void LogPresenceInfo(const EOS_HPlatform Platform, const EOS_EpicAccountId LoggedInAccount, const EOS_EpicAccountId TargetAccount, int32 Indent = 0) const override;
 	virtual void LogFriendsInfo(const EOS_HPlatform Platform, const EOS_EpicAccountId LoggedInAccount, int32 Indent = 0) const override;
 	virtual void LogConnectInfo(const EOS_HPlatform Platform, const EOS_ProductUserId LoggedInAccount, int32 Indent = 0) const override;
+
+	virtual void AddCallbackObject(TUniquePtr<class FCallbackBase> CallbackObj) override;
 	// End IEOSSDKManager
 
 	// Begin FSelfRegisteringExec
@@ -83,6 +85,8 @@ private:
 	TArray<EOS_HPlatform> ReleasedPlatforms;
 	/** Handle to ticker delegate for Tick(), valid whenever there are ActivePlatforms to tick, or ReleasedPlatforms to release. */
 	FTSTicker::FDelegateHandle TickerHandle;
+	/** Callback objects, to be released after EOS_Shutdown */
+	TArray<TUniquePtr<FCallbackBase>> CallbackObjects;
 
 	// Config
 	/** Interval between platform ticks. 0 means we tick every frame. */
diff --git a/Engine/Plugins/Online/EOSShared/Source/EOSShared/Public/IEOSSDKManager.h b/Engine/Plugins/Online/EOSShared/Source/EOSShared/Public/IEOSSDKManager.h
--- a/Engine/Plugins/Online/EOSShared/Source/EOSShared/Public/IEOSSDKManager.h
+++ b/Engine/Plugins/Online/EOSShared/Source/EOSShared/Public/IEOSSDKManager.h
@@ -5,6 +5,7 @@
 #if WITH_EOS_SDK
 
 #include "Features/IModularFeatures.h"
+#include "Templates/UniquePtr.h"
 
 #if defined(EOS_PLATFORM_BASE_FILE_NAME)
 #include EOS_PLATFORM_BASE_FILE_NAME
@@ -79,6 +80,9 @@ public:
 	virtual void LogFriendsInfo(const EOS_HPlatform Platform, const EOS_EpicAccountId LoggedInAccount, int32 Indent = 0) const = 0;
 	virtual void LogConnectInfo(const EOS_HPlatform Platform, const EOS_ProductUserId LoggedInAccount, int32 Indent = 0) const = 0;
 
+	/** Assign ownership of a callback object, to be free'd after EOS_Shutdown */
+	virtual void AddCallbackObject(TUniquePtr<class FCallbackBase> CallbackObj) = 0;
+
 	FEOSSDKManagerOnPreInitializeSDK OnPreInitializeSDK;
 	FEOSSDKManagerOnPreCreatePlatform OnPreCreatePlatform;
 };
diff --git a/Engine/Plugins/Online/VoiceChat/EOSVoiceChat/Source/EOSVoiceChat/Private/EOSVoiceChatUser.cpp b/Engine/Plugins/Online/VoiceChat/EOSVoiceChat/Source/EOSVoiceChat/Private/EOSVoiceChatUser.cpp
--- a/Engine/Plugins/Online/VoiceChat/EOSVoiceChat/Source/EOSVoiceChat/Private/EOSVoiceChatUser.cpp
+++ b/Engine/Plugins/Online/VoiceChat/EOSVoiceChat/Source/EOSVoiceChat/Private/EOSVoiceChatUser.cpp
@@ -9,12 +9,14 @@
 #include "ProfilingDebugging/CsvProfiler.h"
 
 #include "EOSShared.h"
+#include "EOSSharedTypes.h"
 #include "EOSVoiceChat.h"
 #include "EOSVoiceChatErrors.h"
 #include "EOSVoiceChatModule.h"
 #include "EOSVoiceChatTypes.h"
 #include "VoiceChatErrors.h"
 #include "EOSAudioDevicePool.h"
+#include "IEOSSDKManager.h"
 
 #include "eos_common.h"
 #include "eos_lobby.h"
@@ -78,6 +80,8 @@ namespace
 			return TEXT("Unknown");
 		}
 	}
+
+	using FAudioBeforeSendCallback = TEOSGlobalCallback<EOS_RTCAudio_OnAudioBeforeSendCallback, EOS_RTCAudio_AudioBeforeSendCallbackInfo, FEOSVoiceChatUser>;
 }
 
 static TAutoConsoleVariable<bool> CVarFakeAudioInputEnabled(
@@ -99,6 +103,7 @@ FEOSVoiceChatUser::FEOSVoiceChatUser(FEOSVoiceChat& InEOSVoiceChat)
 
 FEOSVoiceChatUser::~FEOSVoiceChatUser()
 {
+	bInDestructor = true;
 	ClearLoginSession();
 }
 
@@ -1395,11 +1400,21 @@ void FEOSVoiceChatUser::BindChannelCallbacks(FChannelSession& ChannelSession)
 		static_assert(EOS_RTCAUDIO_ADDNOTIFYAUDIOBEFORESEND_API_LATEST == 1, "EOS_RTC_AddNotifyAudioBeforeSendOptions updated, check new fields");
 		AudioBeforeSendOptions.LocalUserId = LoginSession.LocalUserProductUserId;
 		AudioBeforeSendOptions.RoomName = Utf8RoomName.Get();
-		ChannelSession.OnAudioBeforeSendNotificationId = EOS_RTCAudio_AddNotifyAudioBeforeSend(EOS_RTC_GetAudioInterface(GetRtcInterface()), &AudioBeforeSendOptions, this, &FEOSVoiceChatUser::OnChannelAudioBeforeSendStatic);
+
+		// Protect against callbacks occurring after this object is destroyed, by wrapping in a TEOSGlobalCallback. This can occur when LeaveRoom during Logout fails.
+		TUniquePtr<FAudioBeforeSendCallback> Callback = MakeUnique<FAudioBeforeSendCallback>(AsWeak());
+		Callback->CallbackLambda = [this](const EOS_RTCAudio_AudioBeforeSendCallbackInfo* Data) { OnChannelAudioBeforeSend(Data); };
+		Callback->bIsGameThreadCallback = false;
+
+		ChannelSession.OnAudioBeforeSendNotificationId = EOS_RTCAudio_AddNotifyAudioBeforeSend(EOS_RTC_GetAudioInterface(GetRtcInterface()), &AudioBeforeSendOptions, Callback.Get(), Callback->GetCallbackPtr());
 		if (ChannelSession.OnAudioBeforeSendNotificationId == EOS_INVALID_NOTIFICATIONID)
 		{
 			EOSVOICECHATUSER_LOG(Warning, TEXT("BindChannelCallbacks EOS_RTCAudio_AddNotifyAudioBeforeSend failed"));
 		}
+		else
+		{
+			ChannelSession.AudioBeforeSendCallback = MoveTemp(Callback);
+		}
 	}
 
 	// OnAudioBeforeRender
@@ -1458,6 +1473,19 @@ void FEOSVoiceChatUser::UnbindChannelCallbacks(FChannelSession& ChannelSession)
 	{
 		EOS_RTCAudio_RemoveNotifyAudioBeforeSend(EOS_RTC_GetAudioInterface(GetRtcInterface()), ChannelSession.OnAudioBeforeSendNotificationId);
 		ChannelSession.OnAudioBeforeSendNotificationId = EOS_INVALID_NOTIFICATIONID;
+
+		if (ChannelSession.AudioBeforeSendCallback)
+		{
+			if (bInDestructor)
+			{
+				// After an OnLeaveRoom callback, it is guaranteed this callback will no longer be fired. This method is in almost all cases called
+				// in response to an OnLeaveRoom callback, so it's safe to release it here. The one exception, is if we are being destroyed. In that
+				// case, the callback may be fired even after unbinding, so we move ownership of the callback to EOSShared so it can outlive this,
+				// and be cleaned up eventually after EOS_Shutdown.
+				EOSVoiceChat.SDKManager.AddCallbackObject(MoveTemp(ChannelSession.AudioBeforeSendCallback));
+			}
+			ChannelSession.AudioBeforeSendCallback.Reset();
+		}
 	}
 
 	if (ChannelSession.OnAudioBeforeRenderNotificationId != EOS_INVALID_NOTIFICATIONID)
@@ -2205,25 +2233,6 @@ void FEOSVoiceChatUser::OnChannelParticipantAudioUpdated(const EOS_RTCAudio_Part
 	}
 }
 
-void EOS_CALL FEOSVoiceChatUser::OnChannelAudioBeforeSendStatic(const EOS_RTCAudio_AudioBeforeSendCallbackInfo* CallbackInfo)
-{
-	if (CallbackInfo)
-	{
-		if (FEOSVoiceChatUser* EosVoiceChatPtr = static_cast<FEOSVoiceChatUser*>(CallbackInfo->ClientData))
-		{
-			EosVoiceChatPtr->OnChannelAudioBeforeSend(CallbackInfo);
-		}
-		else
-		{
-			UE_LOG(LogEOSVoiceChat, Warning, TEXT("OnChannelAudioBeforeSendStatic Error EosVoiceChatPtr=nullptr"));
-		}
-	}
-	else
-	{
-		UE_LOG(LogEOSVoiceChat, Warning, TEXT("OnChannelAudioBeforeSendStatic Error CallbackInfo=nullptr"));
-	}
-}
-
 void FEOSVoiceChatUser::OnChannelAudioBeforeSend(const EOS_RTCAudio_AudioBeforeSendCallbackInfo* CallbackInfo)
 {
 	if (EOS_RTCAudio_AudioBuffer* Buffer = CallbackInfo->Buffer)
diff --git a/Engine/Plugins/Online/VoiceChat/EOSVoiceChat/Source/EOSVoiceChat/Public/EOSVoiceChatUser.h b/Engine/Plugins/Online/VoiceChat/EOSVoiceChat/Source/EOSVoiceChat/Public/EOSVoiceChatUser.h
--- a/Engine/Plugins/Online/VoiceChat/EOSVoiceChat/Source/EOSVoiceChat/Public/EOSVoiceChatUser.h
+++ b/Engine/Plugins/Online/VoiceChat/EOSVoiceChat/Source/EOSVoiceChat/Public/EOSVoiceChatUser.h
@@ -17,6 +17,8 @@ public:
 	FEOSVoiceChatUser(FEOSVoiceChat& InEOSVoiceChat);
 	virtual ~FEOSVoiceChatUser();
 
+	FEOSVoiceChatUser(const FEOSVoiceChatUser&) = delete;
+
 	// ~Begin IVoiceChatUser Interface
 	virtual void SetSetting(const FString& Name, const FString& Value) override;
 	virtual FString GetSetting(const FString& Name) override;
@@ -146,8 +148,12 @@ protected:
 	// Representation of a particular channel
 	struct FChannelSession
 	{
+		FChannelSession() = default;
+		FChannelSession(FChannelSession&&) = default;
 		~FChannelSession();
 
+		FChannelSession(const FChannelSession&) = delete;
+
 		bool IsLocalUser(const FChannelParticipant& Participant);
 		bool IsLobbySession() const;
 		
@@ -190,6 +196,8 @@ protected:
 		EOS_NotificationId OnAudioBeforeSendNotificationId = EOS_INVALID_NOTIFICATIONID;
 		EOS_NotificationId OnAudioBeforeRenderNotificationId = EOS_INVALID_NOTIFICATIONID;
 		EOS_NotificationId OnAudioInputStateNotificationId = EOS_INVALID_NOTIFICATIONID;
+
+		TUniquePtr<class FCallbackBase> AudioBeforeSendCallback;
 	};
 
 	enum class ELoginState
@@ -203,8 +211,12 @@ protected:
 	// When logged in, contains the state for the current login session. Reset by Logout
 	struct FLoginSession
 	{
+		FLoginSession() = default;
+		FLoginSession& operator=(FLoginSession&&) = default;
 		~FLoginSession();
 
+		FLoginSession(const FLoginSession&) = delete;
+
 		// The numeric platform id for the local user
 		FPlatformUserId PlatformId = PLATFORMUSERID_NONE;
 		// Name of the local player
@@ -288,6 +300,7 @@ protected:
 	FEOSVoiceChat& EOSVoiceChat;
 
 	bool bFakeAudioInput = false;
+	bool bInDestructor = false;
 
 	// Helper methods
 	bool IsInitialized();
@@ -345,7 +358,6 @@ protected:
 	void OnChannelParticipantStatusChanged(const EOS_RTC_ParticipantStatusChangedCallbackInfo* CallbackInfo);
 	static void EOS_CALL OnChannelParticipantAudioUpdatedStatic(const EOS_RTCAudio_ParticipantUpdatedCallbackInfo* CallbackInfo);
 	void OnChannelParticipantAudioUpdated(const EOS_RTCAudio_ParticipantUpdatedCallbackInfo* CallbackInfo);
-	static void EOS_CALL OnChannelAudioBeforeSendStatic(const EOS_RTCAudio_AudioBeforeSendCallbackInfo* CallbackInfo);
 	void OnChannelAudioBeforeSend(const EOS_RTCAudio_AudioBeforeSendCallbackInfo* CallbackInfo);
 	static void EOS_CALL OnChannelAudioBeforeRenderStatic(const EOS_RTCAudio_AudioBeforeRenderCallbackInfo* CallbackInfo);
 	void OnChannelAudioBeforeRender(const EOS_RTCAudio_AudioBeforeRenderCallbackInfo* CallbackInfo);
