- Fix crash when logging in P2 in development builds, calling SetPlayerMuted on P1 while not yet logged in.

diff --git a/Engine/Plugins/Online/VoiceChat/EOSVoiceChat/Source/EOSVoiceChat/Private/EOSVoiceChatUser.cpp b/Engine/Plugins/Online/VoiceChat/EOSVoiceChat/Source/EOSVoiceChat/Private/EOSVoiceChatUser.cpp
--- a/Engine/Plugins/Online/VoiceChat/EOSVoiceChat/Source/EOSVoiceChat/Private/EOSVoiceChatUser.cpp
+++ b/Engine/Plugins/Online/VoiceChat/EOSVoiceChat/Source/EOSVoiceChat/Private/EOSVoiceChatUser.cpp
@@ -701,6 +701,12 @@ bool FEOSVoiceChatUser::IsPlayerTalking(const FString& PlayerName) const
 
 void FEOSVoiceChatUser::SetPlayerMuted(const FString& PlayerName, bool bAudioMuted)
 {
+	if (!IsLoggedIn())
+	{
+		EOSVOICECHATUSER_LOG(Warning, TEXT("SetPlayerMuted called while logged out"));
+		return;
+	}
+
 	FGlobalParticipant& GlobalParticipant = GetGlobalParticipant(PlayerName);
 	const bool bAudioMutedChanged = GlobalParticipant.bAudioMuted != bAudioMuted;
 
@@ -717,6 +723,12 @@ void FEOSVoiceChatUser::SetPlayerMuted(const FString& PlayerName, bool bAudioMut
 
 bool FEOSVoiceChatUser::IsPlayerMuted(const FString& PlayerName) const
 {
+	if (!IsLoggedIn())
+	{
+		EOSVOICECHATUSER_LOG(Warning, TEXT("IsPlayerMuted called while logged out"));
+		return false;
+	}
+
 	return GetGlobalParticipant(PlayerName).bAudioMuted;
 }
 
