- Wrap UpdateSendingAudio in TEOSCallback, to protect against EOSVoiceChatUser being destroyed while the update is in flight.

diff --git a/Engine/Plugins/Online/EOSShared/Source/EOSShared/Public/EOSSharedTypes.h b/Engine/Plugins/Online/EOSShared/Source/EOSShared/Public/EOSSharedTypes.h
--- a/Engine/Plugins/Online/EOSShared/Source/EOSShared/Public/EOSSharedTypes.h
+++ b/Engine/Plugins/Online/EOSShared/Source/EOSShared/Public/EOSSharedTypes.h
@@ -26,7 +26,7 @@ public:
 
 #if WITH_EOS_SDK
 
-/** Class to handle all callbacks generically using a lambda to process callback results */
+/** Wrapper for EOS notification callbacks, ensures calling object is still alive before calling passed lambda. */
 template<typename CallbackFuncType, typename CallbackParamType, typename OwningType, typename CallbackReturnType = void, typename... CallbackExtraParams>
 class TEOSGlobalCallback :
 	public FCallbackBase
@@ -85,4 +85,55 @@ private:
 	}
 };
 
+/** Wrapper for EOS async API callbacks, ensures the caller is still alive before calling passed lambda, then deletes itself. */
+template<typename CallbackFuncType, typename CallbackType, typename OwningType>
+class TEOSCallback :
+	public FCallbackBase
+{
+public:
+	TFunction<void(const CallbackType*)> CallbackLambda;
+
+	TEOSCallback(TWeakPtr<OwningType> InOwner)
+		: FCallbackBase()
+		, Owner(InOwner)
+	{
+	}
+	TEOSCallback(TWeakPtr<const OwningType> InOwner)
+		: FCallbackBase()
+		, Owner(InOwner)
+	{
+	}
+	virtual ~TEOSCallback() = default;
+
+	CallbackFuncType GetCallbackPtr()
+	{
+		return &CallbackImpl;
+	}
+
+protected:
+	/** The object that needs to be checked for lifetime before calling the callback */
+	TWeakPtr<const OwningType> Owner;
+
+private:
+	static void EOS_CALL CallbackImpl(const CallbackType* Data)
+	{
+		if (EOS_EResult_IsOperationComplete(Data->ResultCode) == EOS_FALSE)
+		{
+			// Ignore
+			return;
+		}
+		check(IsInGameThread());
+
+		TEOSCallback* CallbackThis = (TEOSCallback*)Data->ClientData;
+		check(CallbackThis);
+
+		if (CallbackThis->Owner.IsValid())
+		{
+			check(CallbackThis->CallbackLambda);
+			CallbackThis->CallbackLambda(Data);
+		}
+		delete CallbackThis;
+	}
+};
+
 #endif
\ No newline at end of file
diff --git a/Engine/Plugins/Online/OnlineSubsystemEOS/Source/OnlineSubsystemEOS/Private/OnlineSubsystemEOSTypes.h b/Engine/Plugins/Online/OnlineSubsystemEOS/Source/OnlineSubsystemEOS/Private/OnlineSubsystemEOSTypes.h
--- a/Engine/Plugins/Online/OnlineSubsystemEOS/Source/OnlineSubsystemEOS/Private/OnlineSubsystemEOSTypes.h
+++ b/Engine/Plugins/Online/OnlineSubsystemEOS/Source/OnlineSubsystemEOS/Private/OnlineSubsystemEOSTypes.h
@@ -296,58 +296,6 @@ protected:
 	EInviteStatus::Type InviteStatus;
 };
 
-/** Class to handle all callbacks generically using a lambda to process callback results */
-template<typename CallbackFuncType, typename CallbackType, typename OwningType>
-class TEOSCallback :
-	public FCallbackBase
-{
-public:
-	TFunction<void(const CallbackType*)> CallbackLambda;
-
-	TEOSCallback(TWeakPtr<OwningType> InOwner)
-		: FCallbackBase()
-		, Owner(InOwner)
-	{
-	}
-	TEOSCallback(TWeakPtr<const OwningType> InOwner)
-		: FCallbackBase()
-		, Owner(InOwner)
-	{
-	}
-	virtual ~TEOSCallback() = default;
-
-
-	CallbackFuncType GetCallbackPtr()
-	{
-		return &CallbackImpl;
-	}
-
-protected:
-	/** The object that needs to be checked for lifetime before calling the callback */
-	TWeakPtr<const OwningType> Owner;
-
-private:
-	static void EOS_CALL CallbackImpl(const CallbackType* Data)
-	{
-		if (EOS_EResult_IsOperationComplete(Data->ResultCode) == EOS_FALSE)
-		{
-			// Ignore
-			return;
-		}
-		check(IsInGameThread());
-
-		TEOSCallback* CallbackThis = (TEOSCallback*)Data->ClientData;
-		check(CallbackThis);
-
-		if (CallbackThis->Owner.IsValid())
-		{
-			check(CallbackThis->CallbackLambda);
-			CallbackThis->CallbackLambda(Data);
-		}
-		delete CallbackThis;
-	}
-};
-
 namespace OSSInternalCallback
 {
 	/** Create a callback for a non-SDK function that is tied to the lifetime of an arbitrary shared pointer. */
diff --git a/Engine/Plugins/Online/VoiceChat/EOSVoiceChat/Source/EOSVoiceChat/Private/EOSVoiceChatUser.cpp b/Engine/Plugins/Online/VoiceChat/EOSVoiceChat/Source/EOSVoiceChat/Private/EOSVoiceChatUser.cpp
--- a/Engine/Plugins/Online/VoiceChat/EOSVoiceChat/Source/EOSVoiceChat/Private/EOSVoiceChatUser.cpp
+++ b/Engine/Plugins/Online/VoiceChat/EOSVoiceChat/Source/EOSVoiceChat/Private/EOSVoiceChatUser.cpp
@@ -26,11 +26,8 @@
 #include "eos_rtc_audio.h"
 #include "eos_sdk.h"
 
-
 #define EOS_VOICE_TODO 0
 
-#define CHECKPIN() FEOSVoiceChatUserPtr StrongThis = WeakThis.Pin(); if(!StrongThis) return
-
 namespace
 {
 	template<class TDelegate, class... TArgs>
@@ -84,6 +81,7 @@ namespace
 	}
 
 	using FAudioBeforeSendCallback = TEOSGlobalCallback<EOS_RTCAudio_OnAudioBeforeSendCallback, EOS_RTCAudio_AudioBeforeSendCallbackInfo, FEOSVoiceChatUser>;
+	using FUpdateSendingAudioCallback = TEOSCallback<EOS_RTCAudio_OnUpdateSendingCallback, EOS_RTCAudio_UpdateSendingCallbackInfo, FEOSVoiceChatUser>;
 }
 
 static TAutoConsoleVariable<bool> CVarFakeAudioInputEnabled(
@@ -329,7 +327,7 @@ FVoiceChatDeviceInfo FEOSVoiceChatUser::GetDefaultOutputDeviceInfo() const
 	{
 		DefaultDeviceInfo = EOSVoiceChat.InitSession.EosAudioDevicePool->GetCachedOutputDeviceInfos()[EOSVoiceChat.InitSession.EosAudioDevicePool->GetDefaultOutputDeviceInfoIdx()];
 	}
-	 
+
 	return DefaultDeviceInfo;
 }
 
@@ -1328,9 +1326,13 @@ void FEOSVoiceChatUser::ApplySendingOptions(FChannelSession& ChannelSession)
 		? EOS_ERTCAudioStatus::EOS_RTCAS_Enabled
 		: EOS_ERTCAudioStatus::EOS_RTCAS_Disabled;
 
+	// Protect against the callback firing after this object is destroyed.
+	TUniquePtr<FUpdateSendingAudioCallback> Callback = MakeUnique<FUpdateSendingAudioCallback>(AsWeak());
+	Callback->CallbackLambda = [this](const EOS_RTCAudio_UpdateSendingCallbackInfo* Data) { OnUpdateSendingAudio(Data); };
+
 	CSV_SCOPED_TIMING_STAT_EXCLUSIVE(EOSVoiceChat);
 	QUICK_SCOPE_CYCLE_COUNTER(EOS_RTCAudio_UpdateSending);
-	EOS_RTCAudio_UpdateSending(EOS_RTC_GetAudioInterface(GetRtcInterface()), &UpdateSendingOptions, this, &FEOSVoiceChatUser::OnUpdateSendingAudioStatic);
+	EOS_RTCAudio_UpdateSending(EOS_RTC_GetAudioInterface(GetRtcInterface()), &UpdateSendingOptions, Callback.Get(), Callback->GetCallbackPtr());
 }
 
 void FEOSVoiceChatUser::BindLoginCallbacks()
@@ -1945,25 +1947,6 @@ void FEOSVoiceChatUser::OnUpdateReceivingAudio(const EOS_RTCAudio_UpdateReceivin
 	}
 }
 
-void EOS_CALL FEOSVoiceChatUser::OnUpdateSendingAudioStatic(const EOS_RTCAudio_UpdateSendingCallbackInfo* CallbackInfo)
-{
-	if (CallbackInfo)
-	{
-		if (FEOSVoiceChatUser* EosVoiceChatPtr = static_cast<FEOSVoiceChatUser*>(CallbackInfo->ClientData))
-		{
-			EosVoiceChatPtr->OnUpdateSendingAudio(CallbackInfo);
-		}
-		else
-		{
-			UE_LOG(LogEOSVoiceChat, Warning, TEXT("OnUpdateSendingStatic Error EosVoiceChatPtr=nullptr"));
-		}
-	}
-	else
-	{
-		UE_LOG(LogEOSVoiceChat, Warning, TEXT("OnUpdateSendingStatic Error CallbackInfo=nullptr"));
-	}
-}
-
 void FEOSVoiceChatUser::OnUpdateSendingAudio(const EOS_RTCAudio_UpdateSendingCallbackInfo* CallbackInfo)
 {
 	check(IsInitialized());
@@ -2754,6 +2737,4 @@ const TCHAR* LexToString(FEOSVoiceChatUser::EChannelJoinState State)
 	}
 }
 
-#undef CHECKPIN
-
 #endif // WITH_EOS_RTC
\ No newline at end of file
diff --git a/Engine/Plugins/Online/VoiceChat/EOSVoiceChat/Source/EOSVoiceChat/Public/EOSVoiceChatUser.h b/Engine/Plugins/Online/VoiceChat/EOSVoiceChat/Source/EOSVoiceChat/Public/EOSVoiceChatUser.h
--- a/Engine/Plugins/Online/VoiceChat/EOSVoiceChat/Source/EOSVoiceChat/Public/EOSVoiceChatUser.h
+++ b/Engine/Plugins/Online/VoiceChat/EOSVoiceChat/Source/EOSVoiceChat/Public/EOSVoiceChatUser.h
@@ -346,7 +346,6 @@ protected:
 	void OnUpdateParticipantVolume(const EOS_RTCAudio_UpdateParticipantVolumeCallbackInfo* CallbackInfo);
 	static void EOS_CALL OnUpdateReceivingAudioStatic(const EOS_RTCAudio_UpdateReceivingCallbackInfo* CallbackInfo);
 	void OnUpdateReceivingAudio(const EOS_RTCAudio_UpdateReceivingCallbackInfo* CallbackInfo);
-	static void EOS_CALL OnUpdateSendingAudioStatic(const EOS_RTCAudio_UpdateSendingCallbackInfo* CallbackInfo);
 	void OnUpdateSendingAudio(const EOS_RTCAudio_UpdateSendingCallbackInfo* CallbackInfo);
 
 	// EOS notification callbacks
