- Fix not releasing the unique ptr, so it's immediately destroyed and crashes in the callback.

diff --git a/Engine/Plugins/Online/VoiceChat/EOSVoiceChat/Source/EOSVoiceChat/Private/EOSVoiceChatUser.cpp b/Engine/Plugins/Online/VoiceChat/EOSVoiceChat/Source/EOSVoiceChat/Private/EOSVoiceChatUser.cpp
--- a/Engine/Plugins/Online/VoiceChat/EOSVoiceChat/Source/EOSVoiceChat/Private/EOSVoiceChatUser.cpp
+++ b/Engine/Plugins/Online/VoiceChat/EOSVoiceChat/Source/EOSVoiceChat/Private/EOSVoiceChatUser.cpp
@@ -1332,7 +1332,7 @@ void FEOSVoiceChatUser::ApplySendingOptions(FChannelSession& ChannelSession)
 
 	CSV_SCOPED_TIMING_STAT_EXCLUSIVE(EOSVoiceChat);
 	QUICK_SCOPE_CYCLE_COUNTER(EOS_RTCAudio_UpdateSending);
-	EOS_RTCAudio_UpdateSending(EOS_RTC_GetAudioInterface(GetRtcInterface()), &UpdateSendingOptions, Callback.Get(), Callback->GetCallbackPtr());
+	EOS_RTCAudio_UpdateSending(EOS_RTC_GetAudioInterface(GetRtcInterface()), &UpdateSendingOptions, Callback.Release(), Callback->GetCallbackPtr());
 }
 
 void FEOSVoiceChatUser::BindLoginCallbacks()
