- Fix relying on a specific order of evaluation, but they can be in any order.

diff --git a/Engine/Plugins/Online/VoiceChat/EOSVoiceChat/Source/EOSVoiceChat/Private/EOSVoiceChatUser.cpp b/Engine/Plugins/Online/VoiceChat/EOSVoiceChat/Source/EOSVoiceChat/Private/EOSVoiceChatUser.cpp
--- a/Engine/Plugins/Online/VoiceChat/EOSVoiceChat/Source/EOSVoiceChat/Private/EOSVoiceChatUser.cpp
+++ b/Engine/Plugins/Online/VoiceChat/EOSVoiceChat/Source/EOSVoiceChat/Private/EOSVoiceChatUser.cpp
@@ -1327,12 +1327,12 @@ void FEOSVoiceChatUser::ApplySendingOptions(FChannelSession& ChannelSession)
 		: EOS_ERTCAudioStatus::EOS_RTCAS_Disabled;
 
 	// Protect against the callback firing after this object is destroyed.
-	TUniquePtr<FUpdateSendingAudioCallback> Callback = MakeUnique<FUpdateSendingAudioCallback>(AsWeak());
+	FUpdateSendingAudioCallback* Callback = new FUpdateSendingAudioCallback(AsWeak());
 	Callback->CallbackLambda = [this](const EOS_RTCAudio_UpdateSendingCallbackInfo* Data) { OnUpdateSendingAudio(Data); };
 
 	CSV_SCOPED_TIMING_STAT_EXCLUSIVE(EOSVoiceChat);
 	QUICK_SCOPE_CYCLE_COUNTER(EOS_RTCAudio_UpdateSending);
-	EOS_RTCAudio_UpdateSending(EOS_RTC_GetAudioInterface(GetRtcInterface()), &UpdateSendingOptions, Callback.Release(), Callback->GetCallbackPtr());
+	EOS_RTCAudio_UpdateSending(EOS_RTC_GetAudioInterface(GetRtcInterface()), &UpdateSendingOptions, Callback, Callback->GetCallbackPtr());
 }
 
 void FEOSVoiceChatUser::BindLoginCallbacks()
