FExec: Wrapped remaining Exec overrides with UE_ALLOW_EXEC_COMMANDS so that we can compile the engine without exec command support when desired.

diff --git a/Engine/Plugins/Online/OnlineFramework/Source/Party/Private/SocialDebugTools.cpp b/Engine/Plugins/Online/OnlineFramework/Source/Party/Private/SocialDebugTools.cpp
--- a/Engine/Plugins/Online/OnlineFramework/Source/Party/Private/SocialDebugTools.cpp
+++ b/Engine/Plugins/Online/OnlineFramework/Source/Party/Private/SocialDebugTools.cpp
@@ -44,6 +44,7 @@ void USocialDebugTools::PrintExecCommands() const
 	UE_LOG(LogParty, Log, TEXT("HELP"));
 }
 
+#if UE_ALLOW_EXEC_COMMANDS
 bool USocialDebugTools::Exec(class UWorld* InWorld, const TCHAR* Cmd, FOutputDevice& Out)
 {
 	if (FParse::Command(&Cmd, TEXT("DEBUG")))
@@ -89,6 +90,7 @@ bool USocialDebugTools::Exec(class UWorld* InWorld, const TCHAR* Cmd, FOutputDev
 	
 	return false;
 }
+#endif // UE_ALLOW_EXEC_COMMANDS
 
 void USocialDebugTools::Login(const FString& Instance, const FOnlineAccountCredentials& Credentials, const FLoginComplete& OnComplete)
 {
diff --git a/Engine/Plugins/Online/OnlineFramework/Source/Party/Public/SocialDebugTools.h b/Engine/Plugins/Online/OnlineFramework/Source/Party/Public/SocialDebugTools.h
--- a/Engine/Plugins/Online/OnlineFramework/Source/Party/Public/SocialDebugTools.h
+++ b/Engine/Plugins/Online/OnlineFramework/Source/Party/Public/SocialDebugTools.h
@@ -28,7 +28,9 @@ public:
 	USocialManager& GetSocialManager() const;
 
 	// FExec
+#if UE_ALLOW_EXEC_COMMANDS
 	virtual bool Exec(class UWorld* InWorld, const TCHAR* Cmd, FOutputDevice& Out) override;
+#endif
 
 	// USocialDebugTools
 
diff --git a/Engine/Plugins/Online/VoiceChat/EOSVoiceChat/Source/EOSVoiceChat/Private/EOSVoiceChatFactory.cpp b/Engine/Plugins/Online/VoiceChat/EOSVoiceChat/Source/EOSVoiceChat/Private/EOSVoiceChatFactory.cpp
--- a/Engine/Plugins/Online/VoiceChat/EOSVoiceChat/Source/EOSVoiceChat/Private/EOSVoiceChatFactory.cpp
+++ b/Engine/Plugins/Online/VoiceChat/EOSVoiceChat/Source/EOSVoiceChat/Private/EOSVoiceChatFactory.cpp
@@ -58,6 +58,7 @@ IVoiceChatPtr FEOSVoiceChatFactory::CreateInstanceWithPlatform(const IEOSPlatfor
 	return VoiceChat;
 }
 
+#if UE_ALLOW_EXEC_COMMANDS
 bool FEOSVoiceChatFactory::Exec(UWorld* InWorld, const TCHAR* Cmd, FOutputDevice& Ar)
 {
 	// When list is requested, the factory is responsible for returning true when all instances have reported their state.
@@ -91,5 +92,6 @@ bool FEOSVoiceChatFactory::Exec(UWorld* InWorld, const TCHAR* Cmd, FOutputDevice
 
 	return bListRequested;
 }
+#endif // UE_ALLOW_EXEC_COMMANDS
 
 #endif // WITH_EOS_RTC
\ No newline at end of file
diff --git a/Engine/Plugins/Online/VoiceChat/EOSVoiceChat/Source/EOSVoiceChat/Public/EOSVoiceChatFactory.h b/Engine/Plugins/Online/VoiceChat/EOSVoiceChat/Source/EOSVoiceChat/Public/EOSVoiceChatFactory.h
--- a/Engine/Plugins/Online/VoiceChat/EOSVoiceChat/Source/EOSVoiceChat/Public/EOSVoiceChatFactory.h
+++ b/Engine/Plugins/Online/VoiceChat/EOSVoiceChat/Source/EOSVoiceChat/Public/EOSVoiceChatFactory.h
@@ -40,7 +40,9 @@ public:
 	IVoiceChatPtr CreateInstanceWithPlatform(const IEOSPlatformHandlePtr& PlatformHandle);
 
 	// ~Begin FSelfRegisteringExec
+#if UE_ALLOW_EXEC_COMMANDS
 	virtual bool Exec(UWorld* InWorld, const TCHAR* Cmd, FOutputDevice& Ar) override;
+#endif
 	// ~End FSelfRegisteringExec
 
 private:
