- Delete deprecated typedef
- Delete unused member vars
- Delete unexecuted delegate

diff --git a/Engine/Plugins/Online/VoiceChat/EOSVoiceChat/Source/EOSVoiceChat/Private/EOSVoiceChat.cpp b/Engine/Plugins/Online/VoiceChat/EOSVoiceChat/Source/EOSVoiceChat/Private/EOSVoiceChat.cpp
--- a/Engine/Plugins/Online/VoiceChat/EOSVoiceChat/Source/EOSVoiceChat/Private/EOSVoiceChat.cpp
+++ b/Engine/Plugins/Online/VoiceChat/EOSVoiceChat/Source/EOSVoiceChat/Private/EOSVoiceChat.cpp
@@ -42,7 +42,6 @@ const TCHAR* LexToString(EOS_ERTCAudioInputStatus Status)
 }
 
 FEOSVoiceChatDelegates::FOnAudioInputDeviceStatusChanged FEOSVoiceChatDelegates::OnAudioInputDeviceStatusChanged;
-FEOSVoiceChatDelegates::FOnVoiceChatChannelConnectionStateDelegate FEOSVoiceChatDelegates::OnVoiceChatChannelConnectionStateChanged;
 FEOSVoiceChatDelegates::FOnVoiceChatPlayerAddedMetadataDelegate FEOSVoiceChatDelegates::OnVoiceChatPlayerAddedMetadata;
 FEOSVoiceChatDelegates::FOnAudioStatusChanged FEOSVoiceChatDelegates::OnAudioStatusChanged;
 
diff --git a/Engine/Plugins/Online/VoiceChat/EOSVoiceChat/Source/EOSVoiceChat/Public/EOSVoiceChat.h b/Engine/Plugins/Online/VoiceChat/EOSVoiceChat/Source/EOSVoiceChat/Public/EOSVoiceChat.h
--- a/Engine/Plugins/Online/VoiceChat/EOSVoiceChat/Source/EOSVoiceChat/Public/EOSVoiceChat.h
+++ b/Engine/Plugins/Online/VoiceChat/EOSVoiceChat/Source/EOSVoiceChat/Public/EOSVoiceChat.h
@@ -28,9 +28,6 @@ public:
 	/** Delegate called when the status of the audio device has changed. */
 	DECLARE_MULTICAST_DELEGATE_TwoParams(FOnAudioInputDeviceStatusChanged, FString /*PlayerName*/, EOS_ERTCAudioInputStatus /*Status*/);
 	static FOnAudioInputDeviceStatusChanged OnAudioInputDeviceStatusChanged;
-	/** Delegate called when the "connection state" changes (happens when a sidekick device toggles it's "mode" on/off) **/
-	DECLARE_MULTICAST_DELEGATE_FourParams(FOnVoiceChatChannelConnectionStateDelegate, const FString& /*PlayerName*/, const FString& /* ChannelName */, bool /* bIsChannelEnabled */, bool /* bSendAudioEnabled */);
-	static FOnVoiceChatChannelConnectionStateDelegate OnVoiceChatChannelConnectionStateChanged;
 	/** Delegate called when ParticipantMetadata is included in the player joined event **/
 	DECLARE_MULTICAST_DELEGATE_FourParams(FOnVoiceChatPlayerAddedMetadataDelegate, const FString& /*LoginPlayerName*/, const FString& /* ChannelName */, const FString& /*PlayerName*/, const TArray<FVoiceChatMetadataItem>& /* PlayerMetadata */);
 	static FOnVoiceChatPlayerAddedMetadataDelegate OnVoiceChatPlayerAddedMetadata;
diff --git a/Engine/Plugins/Online/VoiceChat/EOSVoiceChat/Source/EOSVoiceChat/Public/EOSVoiceChatTypes.h b/Engine/Plugins/Online/VoiceChat/EOSVoiceChat/Source/EOSVoiceChat/Public/EOSVoiceChatTypes.h
--- a/Engine/Plugins/Online/VoiceChat/EOSVoiceChat/Source/EOSVoiceChat/Public/EOSVoiceChatTypes.h
+++ b/Engine/Plugins/Online/VoiceChat/EOSVoiceChat/Source/EOSVoiceChat/Public/EOSVoiceChatTypes.h
@@ -19,7 +19,4 @@ struct FEOSVoiceChatChannelCredentials
 		JSON_SERIALIZE("client_base_url", ClientBaseUrl);
 		JSON_SERIALIZE("participant_token", ParticipantToken);
 	END_JSON_SERIALIZER
-};
-
-UE_DEPRECATED(4.27, "FEOSVoiceChannelConnectionInfo renamed to FEOSVoiceChatChannelCredentials")
-typedef FEOSVoiceChatChannelCredentials FEOSVoiceChannelConnectionInfo;
\ No newline at end of file
+};
\ No newline at end of file
diff --git a/Engine/Plugins/Online/VoiceChat/EOSVoiceChat/Source/EOSVoiceChat/Public/EOSVoiceChatUser.h b/Engine/Plugins/Online/VoiceChat/EOSVoiceChat/Source/EOSVoiceChat/Public/EOSVoiceChatUser.h
--- a/Engine/Plugins/Online/VoiceChat/EOSVoiceChat/Source/EOSVoiceChat/Public/EOSVoiceChatUser.h
+++ b/Engine/Plugins/Online/VoiceChat/EOSVoiceChat/Source/EOSVoiceChat/Public/EOSVoiceChatUser.h
@@ -167,10 +167,6 @@ protected:
 		FString PlayerName;
 		// Current participants in the channel, and the current blocked/muted state
 		TMap<FString, FChannelParticipant> Participants;
-		// Is the connection enabled
-		bool bConnectionEnabled = false;
-		// Is this connection allowed to send audio? Will be false if houseparty is being used as a mic
-		bool bSendAudioEnabled = true;
 		// Did the user toggle "off" this channel
 		bool bIsNotListening = false;
 
