- Fix EOSVoiceChatUser::OnUpdatedSendingAudio ignoring AudioStatus, causing PTT closing to not update IsTalking

diff --git a/Engine/Plugins/Online/VoiceChat/EOSVoiceChat/Source/EOSVoiceChat/Private/EOSVoiceChatUser.cpp b/Engine/Plugins/Online/VoiceChat/EOSVoiceChat/Source/EOSVoiceChat/Private/EOSVoiceChatUser.cpp
--- a/Engine/Plugins/Online/VoiceChat/EOSVoiceChat/Source/EOSVoiceChat/Private/EOSVoiceChatUser.cpp
+++ b/Engine/Plugins/Online/VoiceChat/EOSVoiceChat/Source/EOSVoiceChat/Private/EOSVoiceChatUser.cpp
@@ -1952,7 +1952,7 @@ void FEOSVoiceChatUser::OnUpdateSendingAudio(const EOS_RTCAudio_UpdateSendingCal
 	check(IsInitialized());
 
 	FString ChannelName = UTF8_TO_TCHAR(CallbackInfo->RoomName);
-	const bool bAudioEnabled = EOS_TRUE;//CallbackInfo->AudioStatus == EOS_ERTCAudioStatus::EOS_RTCAS_Enabled;
+	const bool bAudioEnabled = CallbackInfo->AudioStatus == EOS_ERTCAudioStatus::EOS_RTCAS_Enabled;
 
 	if (CallbackInfo->ResultCode == EOS_EResult::EOS_Success)
 	{
