- Fix failing IsInitialized() check when OnUpdateSendingAudio fires after uninitialization, but while the VoiceChatUser is still valid.

diff --git a/Engine/Plugins/Online/VoiceChat/EOSVoiceChat/Source/EOSVoiceChat/Private/EOSVoiceChatUser.cpp b/Engine/Plugins/Online/VoiceChat/EOSVoiceChat/Source/EOSVoiceChat/Private/EOSVoiceChatUser.cpp
--- a/Engine/Plugins/Online/VoiceChat/EOSVoiceChat/Source/EOSVoiceChat/Private/EOSVoiceChatUser.cpp
+++ b/Engine/Plugins/Online/VoiceChat/EOSVoiceChat/Source/EOSVoiceChat/Private/EOSVoiceChatUser.cpp
@@ -1949,41 +1949,42 @@ void FEOSVoiceChatUser::OnUpdateReceivingAudio(const EOS_RTCAudio_UpdateReceivin
 
 void FEOSVoiceChatUser::OnUpdateSendingAudio(const EOS_RTCAudio_UpdateSendingCallbackInfo* CallbackInfo)
 {
-	check(IsInitialized());
-
-	FString ChannelName = UTF8_TO_TCHAR(CallbackInfo->RoomName);
+	const FString ChannelName = UTF8_TO_TCHAR(CallbackInfo->RoomName);
 	const bool bAudioEnabled = CallbackInfo->AudioStatus == EOS_ERTCAudioStatus::EOS_RTCAS_Enabled;
 
 	if (CallbackInfo->ResultCode == EOS_EResult::EOS_Success)
 	{
 		EOSVOICECHATUSER_LOG(Log, TEXT("OnUpdateSending ChannelName=[%s] Success"), *ChannelName);
 
-		if (FChannelSession* ChannelSession = LoginSession.ChannelSessions.Find(ChannelName))
+		if (IsLoggedIn())
 		{
-			if (bAudioEnabled != ChannelSession->ActiveSendingState.bAudioEnabled)
+			if (FChannelSession* ChannelSession = LoginSession.ChannelSessions.Find(ChannelName))
 			{
-				ChannelSession->ActiveSendingState.bAudioEnabled = bAudioEnabled;
-
-				if (FChannelParticipant* LocalChannelParticipant = ChannelSession->Participants.Find(ChannelSession->PlayerName))
+				if (bAudioEnabled != ChannelSession->ActiveSendingState.bAudioEnabled)
 				{
-					// The "Talking" state exposed via IVoiceChat is bTalking && bAudioEnabled. If bTalking == false there is no change.
-					if (LocalChannelParticipant->bTalking)
+					ChannelSession->ActiveSendingState.bAudioEnabled = bAudioEnabled;
+
+					if (FChannelParticipant* LocalChannelParticipant = ChannelSession->Participants.Find(ChannelSession->PlayerName))
 					{
-						FGlobalParticipant& LocalGlobalParticipant = GetGlobalParticipant(LocalChannelParticipant->PlayerName);
-						LocalGlobalParticipant.bTalking = bAudioEnabled;
-						OnVoiceChatPlayerTalkingUpdatedDelegate.Broadcast(ChannelSession->ChannelName, ChannelSession->PlayerName, LocalGlobalParticipant.bTalking);
+						// The "Talking" state exposed via IVoiceChat is bTalking && bAudioEnabled. If bTalking == false there is no change.
+						if (LocalChannelParticipant->bTalking)
+						{
+							FGlobalParticipant& LocalGlobalParticipant = GetGlobalParticipant(LocalChannelParticipant->PlayerName);
+							LocalGlobalParticipant.bTalking = bAudioEnabled;
+							OnVoiceChatPlayerTalkingUpdatedDelegate.Broadcast(ChannelSession->ChannelName, ChannelSession->PlayerName, LocalGlobalParticipant.bTalking);
+						}
+					}
+					else
+					{
+						// Not a warning as the first call to UpdateSending is before JoinRoom, so before the local participant is added...
+						EOSVOICECHATUSER_LOG(Verbose, TEXT("OnUpdateSending ChannelName=[%s] PlayerName=[%s] not found"), *ChannelName);
 					}
-				}
-				else
-				{
-					// Not a warning as the first call to UpdateSending is before JoinRoom, so before the local participant is added...
-					EOSVOICECHATUSER_LOG(Verbose, TEXT("OnUpdateSending ChannelName=[%s] PlayerName=[%s] not found"), *ChannelName);
 				}
 			}
-		}
-		else
-		{
-			EOSVOICECHATUSER_LOG(Warning, TEXT("OnUpdateSending ChannelName=[%s] not found"), *ChannelName);
+			else
+			{
+				EOSVOICECHATUSER_LOG(Warning, TEXT("OnUpdateSending ChannelName=[%s] not found"), *ChannelName);
+			}
 		}
 	}
 	else
