- Merge fix for use-after-free

diff --git a/Engine/Plugins/Online/VoiceChat/EOSVoiceChat/Source/EOSVoiceChat/Private/EOSVoiceChat.cpp b/Engine/Plugins/Online/VoiceChat/EOSVoiceChat/Source/EOSVoiceChat/Private/EOSVoiceChat.cpp
--- a/Engine/Plugins/Online/VoiceChat/EOSVoiceChat/Source/EOSVoiceChat/Private/EOSVoiceChat.cpp
+++ b/Engine/Plugins/Online/VoiceChat/EOSVoiceChat/Source/EOSVoiceChat/Private/EOSVoiceChat.cpp
@@ -310,20 +310,12 @@ void FEOSVoiceChat::ReleaseUser(IVoiceChatUser* User)
 					UE_LOG(LogEOSVoiceChat, Warning, TEXT("ReleaseUser User=[%p] Logout failed, Result=[%s]"), User, *LexToString(Result))
 				}
 
-				UE_LOG(LogEOSVoiceChat, Log, TEXT("ReleaseUser User=[%p] Removing"), User);
-				VoiceChatUsers.RemoveAll([User](const FEOSVoiceChatUserRef& OtherUser)
-				{
-					return User == &OtherUser.Get();
-				});
+				ScheduleReleaseUser(User);
 			}));
 		}
 		else
 		{
-			UE_LOG(LogEOSVoiceChat, Log, TEXT("ReleaseUser User=[%p] Removing"), User);
-			VoiceChatUsers.RemoveAll([User](const FEOSVoiceChatUserRef& OtherUser)
-			{
-				return User == &OtherUser.Get();
-			});
+			ScheduleReleaseUser(User);
 		}
 	}
 }
@@ -836,6 +828,24 @@ void FEOSVoiceChat::OnAudioDevicesChanged()
 	}));
 }
 
+void FEOSVoiceChat::ScheduleReleaseUser(IVoiceChatUser* User)
+{
+	FEOSVoiceChatUserRef UserRef = static_cast<FEOSVoiceChatUser*>(User)->AsShared();
+	ReleasedVoiceChatUsers.Add(UserRef);
+	VoiceChatUsers.Remove(UserRef);
+
+	AsyncTask(ENamedThreads::GameThread, [this, WeakThis = AsWeak()]()
+	{
+		CHECKPIN();
+
+		while (ReleasedVoiceChatUsers.Num() > 0)
+		{
+			UE_LOG(LogEOSVoiceChat, Log, TEXT("Releasing User=[%p]"), &ReleasedVoiceChatUsers[0].Get());
+			ReleasedVoiceChatUsers.RemoveAtSwap(0);
+		}
+	});
+}
+
 FEOSVoiceChatUser& FEOSVoiceChat::GetVoiceChatUser()
 {
 	if (!SingleUserVoiceChatUser)
diff --git a/Engine/Plugins/Online/VoiceChat/EOSVoiceChat/Source/EOSVoiceChat/Public/EOSVoiceChat.h b/Engine/Plugins/Online/VoiceChat/EOSVoiceChat/Source/EOSVoiceChat/Public/EOSVoiceChat.h
--- a/Engine/Plugins/Online/VoiceChat/EOSVoiceChat/Source/EOSVoiceChat/Public/EOSVoiceChat.h
+++ b/Engine/Plugins/Online/VoiceChat/EOSVoiceChat/Source/EOSVoiceChat/Public/EOSVoiceChat.h
@@ -206,6 +206,8 @@ protected:
 	FOnVoiceChatAvailableAudioDevicesChangedDelegate OnVoiceChatAvailableAudioDevicesChangedDelegate;
 
 	TArray<FEOSVoiceChatUserRef> VoiceChatUsers;
+	TArray<FEOSVoiceChatUserRef> ReleasedVoiceChatUsers;
+	void ScheduleReleaseUser(IVoiceChatUser* User);
 
 	FEOSVoiceChatUser& GetVoiceChatUser();
 	FEOSVoiceChatUser& GetVoiceChatUser() const;
