Fixed GetAvailableInputDeviceInfos calling wrong function
#10071

diff --git a/Engine/Plugins/Online/VoiceChat/EOSVoiceChat/Source/EOSVoiceChat/Private/EOSVoiceChat.cpp b/Engine/Plugins/Online/VoiceChat/EOSVoiceChat/Source/EOSVoiceChat/Private/EOSVoiceChat.cpp
--- a/Engine/Plugins/Online/VoiceChat/EOSVoiceChat/Source/EOSVoiceChat/Private/EOSVoiceChat.cpp
+++ b/Engine/Plugins/Online/VoiceChat/EOSVoiceChat/Source/EOSVoiceChat/Private/EOSVoiceChat.cpp
@@ -374,7 +374,7 @@ bool FEOSVoiceChat::GetAudioOutputDeviceMuted() const
 
 TArray<FVoiceChatDeviceInfo> FEOSVoiceChat::GetAvailableInputDeviceInfos() const
 {
-	return GetVoiceChatUser().GetAvailableOutputDeviceInfos();
+	return GetVoiceChatUser().GetAvailableInputDeviceInfos();
 }
 
 TArray<FVoiceChatDeviceInfo> FEOSVoiceChat::GetAvailableOutputDeviceInfos() const
