- Deprecate old Set3DPosition API, and add default no-op implementation
- Add new simplified Set3DPosition API, add default no-op implementation so licensee impls keep compiling.

diff --git a/Engine/Plugins/Online/OnlineSubsystemEOS/Source/OnlineSubsystemEOS/Private/OnlineSubsystemEOS.cpp b/Engine/Plugins/Online/OnlineSubsystemEOS/Source/OnlineSubsystemEOS/Private/OnlineSubsystemEOS.cpp
--- a/Engine/Plugins/Online/OnlineSubsystemEOS/Source/OnlineSubsystemEOS/Private/OnlineSubsystemEOS.cpp
+++ b/Engine/Plugins/Online/OnlineSubsystemEOS/Source/OnlineSubsystemEOS/Private/OnlineSubsystemEOS.cpp
@@ -87,7 +87,7 @@ public:
 	virtual FOnVoiceChatChannelJoinedDelegate& OnVoiceChatChannelJoined() override { return VoiceChatUser.OnVoiceChatChannelJoined(); }
 	virtual FOnVoiceChatChannelExitedDelegate& OnVoiceChatChannelExited() override { return VoiceChatUser.OnVoiceChatChannelExited(); }
 	virtual FOnVoiceChatCallStatsUpdatedDelegate& OnVoiceChatCallStatsUpdated() override { return VoiceChatUser.OnVoiceChatCallStatsUpdated(); }
-	virtual void Set3DPosition(const FString& ChannelName, const FVector& SpeakerPosition, const FVector& ListenerPosition, const FVector& ListenerForwardDirection, const FVector& ListenerUpDirection) override { VoiceChatUser.Set3DPosition(ChannelName, SpeakerPosition, ListenerPosition, ListenerForwardDirection, ListenerUpDirection); }
+	virtual void Set3DPosition(const FString& ChannelName, const FVector& Position) override { VoiceChatUser.Set3DPosition(ChannelName, Position); }
 	virtual TArray<FString> GetChannels() const override { return VoiceChatUser.GetChannels(); }
 	virtual TArray<FString> GetPlayersInChannel(const FString& ChannelName) const override { return VoiceChatUser.GetPlayersInChannel(ChannelName); }
 	virtual EVoiceChatChannelType GetChannelType(const FString& ChannelName) const override { return VoiceChatUser.GetChannelType(ChannelName); }
diff --git a/Engine/Plugins/Online/VoiceChat/EOSVoiceChat/Source/EOSVoiceChat/Private/EOSVoiceChat.cpp b/Engine/Plugins/Online/VoiceChat/EOSVoiceChat/Source/EOSVoiceChat/Private/EOSVoiceChat.cpp
--- a/Engine/Plugins/Online/VoiceChat/EOSVoiceChat/Source/EOSVoiceChat/Private/EOSVoiceChat.cpp
+++ b/Engine/Plugins/Online/VoiceChat/EOSVoiceChat/Source/EOSVoiceChat/Private/EOSVoiceChat.cpp
@@ -611,9 +611,9 @@ FOnVoiceChatCallStatsUpdatedDelegate& FEOSVoiceChat::OnVoiceChatCallStatsUpdated
 	return GetVoiceChatUser().OnVoiceChatCallStatsUpdated();
 }
 
-void FEOSVoiceChat::Set3DPosition(const FString& ChannelName, const FVector& SpeakerPosition, const FVector& ListenerPosition, const FVector& ListenerForwardDirection, const FVector& ListenerUpDirection)
+void FEOSVoiceChat::Set3DPosition(const FString& ChannelName, const FVector& Position)
 {
-	GetVoiceChatUser().Set3DPosition(ChannelName, SpeakerPosition, ListenerPosition, ListenerForwardDirection, ListenerUpDirection);
+	GetVoiceChatUser().Set3DPosition(ChannelName, Position);
 }
 
 TArray<FString> FEOSVoiceChat::GetChannels() const
diff --git a/Engine/Plugins/Online/VoiceChat/EOSVoiceChat/Source/EOSVoiceChat/Private/EOSVoiceChatUser.cpp b/Engine/Plugins/Online/VoiceChat/EOSVoiceChat/Source/EOSVoiceChat/Private/EOSVoiceChatUser.cpp
--- a/Engine/Plugins/Online/VoiceChat/EOSVoiceChat/Source/EOSVoiceChat/Private/EOSVoiceChatUser.cpp
+++ b/Engine/Plugins/Online/VoiceChat/EOSVoiceChat/Source/EOSVoiceChat/Private/EOSVoiceChatUser.cpp
@@ -649,22 +649,9 @@ void FEOSVoiceChatUser::LeaveChannel(const FString& ChannelName, const FOnVoiceC
 	LeaveChannelInternal(ChannelName, Delegate);
 }
 
-void FEOSVoiceChatUser::Set3DPosition(const FString& ChannelName, const FVector& SpeakerPosition, const FVector& ListenerPosition, const FVector& ListenerForwardDirection, const FVector& ListenerUpDirection)
+void FEOSVoiceChatUser::Set3DPosition(const FString& ChannelName, const FVector& Position)
 {
-#if EOS_VOICE_TODO
-	FChannelSession& ChannelSession = GetChannelSessionGT(ChannelName);
-
-	// Transform Pos and Direction to up -> (0,1,0) and left -> (-1, 0, 0)
-	FVector RotatedPos(ListenerPosition.Y, ListenerPosition.Z, -ListenerPosition.X);
-	FVector RotatedForwardDirection(ListenerForwardDirection.Y, ListenerForwardDirection.Z, -ListenerForwardDirection.X);
-	FVector RotatedUpDirection(ListenerUpDirection.Y, ListenerUpDirection.Z, -ListenerUpDirection.X);
-
-	EOSClientApi::VCSStatus Status = EOSClientConnection.Set3DPosition(LoginSession.AccountName, ChannelSession.ChannelUri, ToEOSVector(SpeakerPosition), ToEOSVector(ListenerPosition), ToEOSVector(ListenerForwardDirection), ToEOSVector(ListenerUpDirection));
-	if (Status.IsError())
-	{
-		EOSVOICECHATUSER_LOG(Warning, TEXT("Set3DPosition failed: channel:%s error:%s (%i)"), ANSI_TO_TCHAR(ChannelSession.ChannelUri), ANSI_TO_TCHAR(Status.ToString()), Status.GetStatusCode());
-	}
-#endif
+	// Unimplemented
 }
 
 TArray<FString> FEOSVoiceChatUser::GetChannels() const
diff --git a/Engine/Plugins/Online/VoiceChat/EOSVoiceChat/Source/EOSVoiceChat/Public/EOSVoiceChat.h b/Engine/Plugins/Online/VoiceChat/EOSVoiceChat/Source/EOSVoiceChat/Public/EOSVoiceChat.h
--- a/Engine/Plugins/Online/VoiceChat/EOSVoiceChat/Source/EOSVoiceChat/Public/EOSVoiceChat.h
+++ b/Engine/Plugins/Online/VoiceChat/EOSVoiceChat/Source/EOSVoiceChat/Public/EOSVoiceChat.h
@@ -109,7 +109,7 @@ public:
 	virtual FOnVoiceChatChannelJoinedDelegate& OnVoiceChatChannelJoined() override;
 	virtual FOnVoiceChatChannelExitedDelegate& OnVoiceChatChannelExited() override;
 	virtual FOnVoiceChatCallStatsUpdatedDelegate& OnVoiceChatCallStatsUpdated() override;
-	virtual void Set3DPosition(const FString& ChannelName, const FVector& SpeakerPosition, const FVector& ListenerPosition, const FVector& ListenerForwardDirection, const FVector& ListenerUpDirection) override;
+	virtual void Set3DPosition(const FString& ChannelName, const FVector& Position) override;
 	virtual TArray<FString> GetChannels() const override;
 	virtual TArray<FString> GetPlayersInChannel(const FString& ChannelName) const override;
 	virtual EVoiceChatChannelType GetChannelType(const FString& ChannelName) const override;
diff --git a/Engine/Plugins/Online/VoiceChat/EOSVoiceChat/Source/EOSVoiceChat/Public/EOSVoiceChatUser.h b/Engine/Plugins/Online/VoiceChat/EOSVoiceChat/Source/EOSVoiceChat/Public/EOSVoiceChatUser.h
--- a/Engine/Plugins/Online/VoiceChat/EOSVoiceChat/Source/EOSVoiceChat/Public/EOSVoiceChatUser.h
+++ b/Engine/Plugins/Online/VoiceChat/EOSVoiceChat/Source/EOSVoiceChat/Public/EOSVoiceChatUser.h
@@ -53,7 +53,7 @@ public:
 	virtual FOnVoiceChatChannelJoinedDelegate& OnVoiceChatChannelJoined() override { return OnVoiceChatChannelJoinedDelegate; }
 	virtual FOnVoiceChatChannelExitedDelegate& OnVoiceChatChannelExited() override { return OnVoiceChatChannelExitedDelegate; }
 	virtual FOnVoiceChatCallStatsUpdatedDelegate& OnVoiceChatCallStatsUpdated() override { return OnVoiceChatCallStatsUpdatedDelegate; }
-	virtual void Set3DPosition(const FString& ChannelName, const FVector& SpeakerPosition, const FVector& ListenerPosition, const FVector& ListenerForwardDirection, const FVector& ListenerUpDirection) override;
+	virtual void Set3DPosition(const FString& ChannelName, const FVector& Position) override;
 	virtual TArray<FString> GetChannels() const override;
 	virtual TArray<FString> GetPlayersInChannel(const FString& ChannelName) const override;
 	virtual EVoiceChatChannelType GetChannelType(const FString& ChannelName) const override;
diff --git a/Engine/Plugins/Online/VoiceChat/VoiceChat/Source/Public/VoiceChat.h b/Engine/Plugins/Online/VoiceChat/VoiceChat/Source/Public/VoiceChat.h
--- a/Engine/Plugins/Online/VoiceChat/VoiceChat/Source/Public/VoiceChat.h
+++ b/Engine/Plugins/Online/VoiceChat/VoiceChat/Source/Public/VoiceChat.h
@@ -466,7 +466,16 @@ public:
 	 * @param ListenerForwardDirection
 	 * @param ListenerUpDirection
 	 */
-	virtual void Set3DPosition(const FString& ChannelName, const FVector& SpeakerPosition, const FVector& ListenerPosition, const FVector& ListenerForwardDirection, const FVector& ListenerUpDirection) = 0;
+	UE_DEPRECATED(5.4, "Please replace with new Set3DPosition override") // TODO also remove the default impl of Set3DPosition below when removing this.
+	virtual void Set3DPosition(const FString& ChannelName, const FVector& SpeakerPosition, const FVector& ListenerPosition, const FVector& ListenerForwardDirection, const FVector& ListenerUpDirection) {};
+
+	/**
+	 * Set the 3d position of the player
+	 *
+	 * @param ChannelName
+	 * @param Position
+	 */
+	virtual void Set3DPosition(const FString& ChannelName, const FVector& Position) {};
 
 	/**
 	 * Get an array of channels the user is in
