Addition of support for -NoEOS commandline parameter to stop EOS modules from loading.

diff --git a/Engine/Plugins/Online/EOSShared/Source/EOSShared/Private/EOSSharedModule.cpp b/Engine/Plugins/Online/EOSShared/Source/EOSShared/Private/EOSSharedModule.cpp
--- a/Engine/Plugins/Online/EOSShared/Source/EOSShared/Private/EOSSharedModule.cpp
+++ b/Engine/Plugins/Online/EOSShared/Source/EOSShared/Private/EOSSharedModule.cpp
@@ -4,6 +4,7 @@
 
 #include "Features/IModularFeatures.h"
 #include "Misc/ConfigCacheIni.h"
+#include "Misc/CommandLine.h"
 #include "Misc/CoreDelegates.h"
 #include "Modules/ModuleManager.h"
 #include "CoreGlobals.h"
@@ -20,6 +21,11 @@ IMPLEMENT_MODULE(FEOSSharedModule, EOSShared);
 void FEOSSharedModule::StartupModule()
 {
 #if WITH_EOS_SDK
+	if (FParse::Param(FCommandLine::Get(), TEXT("NoEOS")))
+	{
+		return;
+	}
+
 	FCoreDelegates::TSOnConfigSectionsChanged().AddRaw(this, &FEOSSharedModule::OnConfigSectionsChanged);
 	LoadConfig();
 
diff --git a/Engine/Plugins/Online/VoiceChat/EOSVoiceChat/Source/EOSVoiceChat/Private/EOSVoiceChatModule.cpp b/Engine/Plugins/Online/VoiceChat/EOSVoiceChat/Source/EOSVoiceChat/Private/EOSVoiceChatModule.cpp
--- a/Engine/Plugins/Online/VoiceChat/EOSVoiceChat/Source/EOSVoiceChat/Private/EOSVoiceChatModule.cpp
+++ b/Engine/Plugins/Online/VoiceChat/EOSVoiceChat/Source/EOSVoiceChat/Private/EOSVoiceChatModule.cpp
@@ -2,6 +2,7 @@
 
 #include "EOSVoiceChatModule.h"
 
+#include "Misc/CommandLine.h"
 #include "Modules/ModuleManager.h"
 
 #include "IEOSSDKManager.h"
@@ -15,6 +16,11 @@ IMPLEMENT_MODULE(FEOSVoiceChatModule, EOSVoiceChat);
 void FEOSVoiceChatModule::StartupModule()
 {
 #if WITH_EOSVOICECHAT
+	if (FParse::Param(FCommandLine::Get(), TEXT("NoEOS")))
+	{
+		return;
+	}
+
 	const FName EOSSharedModuleName = TEXT("EOSShared");
 	const FName EOSSDKManagerFeatureName = TEXT("EOSSDKManager");
 
