Fix and silence new PVS 7.33 warnings

diff --git a/Engine/Plugins/Animation/ACLPlugin/Source/ACLPlugin/Private/AnimationCompressionLibraryDatabase.cpp b/Engine/Plugins/Animation/ACLPlugin/Source/ACLPlugin/Private/AnimationCompressionLibraryDatabase.cpp
--- a/Engine/Plugins/Animation/ACLPlugin/Source/ACLPlugin/Private/AnimationCompressionLibraryDatabase.cpp
+++ b/Engine/Plugins/Animation/ACLPlugin/Source/ACLPlugin/Private/AnimationCompressionLibraryDatabase.cpp
@@ -869,8 +869,7 @@ bool UAnimationCompressionLibraryDatabase::UpdateVisualFidelityTicker(float Delt
 			// We have highest fidelity
 			// The medium/lowest importance tiers are already streamed in
 
-			checkf(Request.Fidelity == ACLVisualFidelity::Medium, TEXT("Unexpected visual fidelity value"));
-			if (Request.Fidelity != ACLVisualFidelity::Medium)
+			if (!ensureMsgf(Request.Fidelity == ACLVisualFidelity::Medium, TEXT("Unexpected visual fidelity value")))
 			{
 				// Something wrong happened, ignore all change requests
 				FailAllRequests(FidelityChangeRequests);
@@ -962,8 +961,7 @@ bool UAnimationCompressionLibraryDatabase::UpdateVisualFidelityTicker(float Delt
 			// We have lowest fidelity
 			// Nothing is currently streamed in
 
-			checkf(Request.Fidelity == ACLVisualFidelity::Medium, TEXT("Unexpected visual fidelity value"));
-			if (Request.Fidelity != ACLVisualFidelity::Medium)
+			if (!ensureMsgf(Request.Fidelity == ACLVisualFidelity::Medium, TEXT("Unexpected visual fidelity value")))
 			{
 				// Something wrong happened, ignore all change requests
 				FailAllRequests(FidelityChangeRequests);
diff --git a/Engine/Plugins/FX/Niagara/Source/Niagara/Private/NiagaraWorldManager.cpp b/Engine/Plugins/FX/Niagara/Source/Niagara/Private/NiagaraWorldManager.cpp
--- a/Engine/Plugins/FX/Niagara/Source/Niagara/Private/NiagaraWorldManager.cpp
+++ b/Engine/Plugins/FX/Niagara/Source/Niagara/Private/NiagaraWorldManager.cpp
@@ -507,18 +507,20 @@ void FNiagaraWorldManager::OnShutdown()
 	}
 
 	//Should have cleared up all world managers by now.
-	check(WorldManagers.Num() == 0);
-	for (TPair<UWorld*, FNiagaraWorldManager*> Pair : WorldManagers)
+	if (!ensure(WorldManagers.Num() == 0))
 	{
-		FNiagaraWorldManager* WorldMan = Pair.Value;
-		if(ensure(WorldMan))
+		for (TPair<UWorld*, FNiagaraWorldManager*> Pair : WorldManagers)
 		{
-			delete WorldMan;
-		}
+			FNiagaraWorldManager* WorldMan = Pair.Value;
+			if (ensure(WorldMan))
+			{
+				delete WorldMan;
+			}
 		
-		Pair.Value = nullptr;
+			Pair.Value = nullptr;
+		}
+		WorldManagers.Empty();
 	}
-	WorldManagers.Empty();
 }
 
 void FNiagaraWorldManager::AddReferencedObjects(FReferenceCollector& Collector)
diff --git a/Engine/Plugins/Media/AudioCaptureTimecodeProvider/Source/AudioCaptureTimecodeProvider/Private/AudioCaptureTimecodeProvider.cpp b/Engine/Plugins/Media/AudioCaptureTimecodeProvider/Source/AudioCaptureTimecodeProvider/Private/AudioCaptureTimecodeProvider.cpp
--- a/Engine/Plugins/Media/AudioCaptureTimecodeProvider/Source/AudioCaptureTimecodeProvider/Private/AudioCaptureTimecodeProvider.cpp
+++ b/Engine/Plugins/Media/AudioCaptureTimecodeProvider/Source/AudioCaptureTimecodeProvider/Private/AudioCaptureTimecodeProvider.cpp
@@ -245,8 +245,10 @@ bool UAudioCaptureTimecodeProvider::Initialize(class UEngine* InEngine)
 		return false;
 	}
 	
-	check(Implementation == nullptr);
-	delete Implementation; // in case
+	if (!ensure(Implementation == nullptr))
+	{
+		delete Implementation;
+	}
 
 	Implementation = new FLinearTimecodeAudioCaptureCustomTimeStepImplementation(this);
 	bool bInitialized = Implementation->Init();
diff --git a/Engine/Plugins/Media/D3D12VideoDecodersElectra/Source/Private/VideoDecoder_D3D12_Common.h b/Engine/Plugins/Media/D3D12VideoDecodersElectra/Source/Private/VideoDecoder_D3D12_Common.h
--- a/Engine/Plugins/Media/D3D12VideoDecodersElectra/Source/Private/VideoDecoder_D3D12_Common.h
+++ b/Engine/Plugins/Media/D3D12VideoDecodersElectra/Source/Private/VideoDecoder_D3D12_Common.h
@@ -131,8 +131,7 @@ public:
 			if (CompletedValue < Value)
 			{
 				HRESULT Result = Fence->SetEventOnCompletion(Value, EventHandle);
-				check(Result == S_OK);
-				if (Result != S_OK)
+				if (!ensure(Result == S_OK))
 				{
 					return false;
 				}
diff --git a/Engine/Plugins/Media/ElectraPlayer/Source/ElectraPlayerRuntime/Private/Runtime/Player/AdaptiveStreamingPlayerRenderers.cpp b/Engine/Plugins/Media/ElectraPlayer/Source/ElectraPlayerRuntime/Private/Runtime/Player/AdaptiveStreamingPlayerRenderers.cpp
--- a/Engine/Plugins/Media/ElectraPlayer/Source/ElectraPlayerRuntime/Private/Runtime/Player/AdaptiveStreamingPlayerRenderers.cpp
+++ b/Engine/Plugins/Media/ElectraPlayer/Source/ElectraPlayerRuntime/Private/Runtime/Player/AdaptiveStreamingPlayerRenderers.cpp
@@ -422,8 +422,7 @@ UEMediaError FAdaptiveStreamingWrappedRenderer::ReturnAudioBuffer(IBuffer* Buffe
 				UEMediaError Error = ReturnBufferCommon(Buffer, bRender, InSampleProperties);
 				check(Error == UEMEDIA_ERROR_OK);
 				Error = WrappedRenderer->AcquireBuffer(Buffer, 0, NoParams);
-				check(Error == UEMEDIA_ERROR_OK);
-				if (Error != UEMEDIA_ERROR_OK)
+				if (!ensure(Error == UEMEDIA_ERROR_OK))
 				{
 					return Error;
 				}
diff --git a/Engine/Plugins/Media/ElectraUtil/Source/ElectraBase/Private/ParameterDictionary.cpp b/Engine/Plugins/Media/ElectraUtil/Source/ElectraBase/Private/ParameterDictionary.cpp
--- a/Engine/Plugins/Media/ElectraUtil/Source/ElectraBase/Private/ParameterDictionary.cpp
+++ b/Engine/Plugins/Media/ElectraUtil/Source/ElectraBase/Private/ParameterDictionary.cpp
@@ -251,8 +251,7 @@ FVariantValue& FVariantValue::Set(const TArray<uint8>& ArrayValue)
 
 const FString& FVariantValue::GetFString() const
 {
-	check(DataType == EDataType::TypeFString);
-	if (DataType == EDataType::TypeFString)
+	if (ensure(DataType == EDataType::TypeFString))
 	{
 		const FString* Str = reinterpret_cast<const FString*>(&DataBuffer);
 		return *Str;
@@ -266,8 +265,7 @@ const FString& FVariantValue::GetFString() const
 
 const double& FVariantValue::GetDouble() const
 {
-	check(DataType == EDataType::TypeDouble);
-	if (DataType == EDataType::TypeDouble)
+	if (ensure(DataType == EDataType::TypeDouble))
 	{
 		const double* Dbl = reinterpret_cast<const double*>(&DataBuffer);
 		return *Dbl;
@@ -281,8 +279,7 @@ const double& FVariantValue::GetDouble() const
 
 const int64& FVariantValue::GetInt64() const
 {
-	check(DataType == EDataType::TypeInt64);
-	if (DataType == EDataType::TypeInt64)
+	if (ensure(DataType == EDataType::TypeInt64))
 	{
 		const int64* Int = reinterpret_cast<const int64*>(&DataBuffer);
 		return *Int;
@@ -296,8 +293,7 @@ const int64& FVariantValue::GetInt64() const
 
 const bool& FVariantValue::GetBool() const
 {
-	check(DataType == EDataType::TypeBoolean);
-	if (DataType == EDataType::TypeBoolean)
+	if (ensure(DataType == EDataType::TypeBoolean))
 	{
 		const bool* Bool = reinterpret_cast<const bool*>(&DataBuffer);
 		return *Bool;
@@ -311,8 +307,7 @@ const bool& FVariantValue::GetBool() const
 
 const FTimeValue& FVariantValue::GetTimeValue() const
 {
-	check(DataType == EDataType::TypeTimeValue);
-	if (DataType == EDataType::TypeTimeValue)
+	if (ensure(DataType == EDataType::TypeTimeValue))
 	{
 		const FTimeValue* Time = reinterpret_cast<const FTimeValue*>(&DataBuffer);
 		return *Time;
@@ -326,8 +321,7 @@ const FTimeValue& FVariantValue::GetTimeValue() const
 
 void* const & FVariantValue::GetPointer() const
 {
-	check(DataType == EDataType::TypeVoidPointer);
-	if (DataType == EDataType::TypeVoidPointer)
+	if (ensure(DataType == EDataType::TypeVoidPointer))
 	{
 		void** Pointer = (void**)&DataBuffer;
 		return *Pointer;
@@ -341,8 +335,7 @@ void* const & FVariantValue::GetPointer() const
 
 const TArray<uint8>& FVariantValue::GetArray() const
 {
-	check(DataType == EDataType::TypeU8Array);
-	if (DataType == EDataType::TypeU8Array)
+	if (ensure(DataType == EDataType::TypeU8Array))
 	{
 		const TArray<uint8>* Array = reinterpret_cast<const TArray<uint8>*>(&DataBuffer);
 		return *Array;
diff --git a/Engine/Plugins/Media/ElectraUtil/Source/ElectraBase/Public/ParameterDictionary.h b/Engine/Plugins/Media/ElectraUtil/Source/ElectraBase/Public/ParameterDictionary.h
--- a/Engine/Plugins/Media/ElectraUtil/Source/ElectraBase/Public/ParameterDictionary.h
+++ b/Engine/Plugins/Media/ElectraUtil/Source/ElectraBase/Public/ParameterDictionary.h
@@ -82,8 +82,7 @@ public:
 	void* const & GetPointer() const;
 	template<typename T> TSharedPtr<T, ESPMode::ThreadSafe> GetSharedPointer() const
 	{
-		check(DataType == EDataType::TypeSharedPointer);
-		if (DataType == EDataType::TypeSharedPointer)
+		if (ensure(DataType == EDataType::TypeSharedPointer))
 		{
 			const TSharedPtrHolder<T>* Pointer = reinterpret_cast<const TSharedPtrHolder<T>*>(&DataBuffer);
 			return Pointer->Pointer;
diff --git a/Engine/Plugins/Media/ElectraUtil/Source/ElectraSamples/Private/ElectraTextureSamplePC.cpp b/Engine/Plugins/Media/ElectraUtil/Source/ElectraSamples/Private/ElectraTextureSamplePC.cpp
--- a/Engine/Plugins/Media/ElectraUtil/Source/ElectraSamples/Private/ElectraTextureSamplePC.cpp
+++ b/Engine/Plugins/Media/ElectraUtil/Source/ElectraSamples/Private/ElectraTextureSamplePC.cpp
@@ -220,8 +220,7 @@ FRHITexture* FElectraTextureSample::GetTexture() const
 					}
 				}
 
-				check(Res == S_OK);
-				if (Res == S_OK)
+				if (ensure(Res == S_OK))
 				{
 					// Setup a suitable RHI texture (it will also become an additional owner of the data)
 					ETextureCreateFlags Flags = ETextureCreateFlags::ShaderResource;
diff --git a/Engine/Plugins/Media/HardwareEncoders/Source/EncoderNVENC/Private/NVENC_Common.cpp b/Engine/Plugins/Media/HardwareEncoders/Source/EncoderNVENC/Private/NVENC_Common.cpp
--- a/Engine/Plugins/Media/HardwareEncoders/Source/EncoderNVENC/Private/NVENC_Common.cpp
+++ b/Engine/Plugins/Media/HardwareEncoders/Source/EncoderNVENC/Private/NVENC_Common.cpp
@@ -83,8 +83,7 @@ namespace AVEncoder
 				{
 					version = NV_ENCODE_API_FUNCTION_LIST_VER;
 					NVENCSTATUS Result = NvEncodeAPICreateInstanceFunc(this);
-					checkf(Result == NV_ENC_SUCCESS, TEXT("Unable to create NvEnc API function list: error %d"), Result);
-					if(Result == NV_ENC_SUCCESS)
+					if (ensureMsgf(Result == NV_ENC_SUCCESS, TEXT("Unable to create NvEnc API function list: error %d"), Result))
 					{
 						bIsAvailable = true;
 					}
diff --git a/Engine/Plugins/Mutable/Source/CustomizableObject/Private/MuCO/CustomizableObject.cpp b/Engine/Plugins/Mutable/Source/CustomizableObject/Private/MuCO/CustomizableObject.cpp
--- a/Engine/Plugins/Mutable/Source/CustomizableObject/Private/MuCO/CustomizableObject.cpp
+++ b/Engine/Plugins/Mutable/Source/CustomizableObject/Private/MuCO/CustomizableObject.cpp
@@ -1230,9 +1230,7 @@ void UCustomizableObjectPrivate::LoadEmbeddedData(FArchive& Ar)
 
 	// If this fails, something went wrong with the packaging: we have data that belongs
 	// to a different version than the code.
-	check(CurrentSupportedVersion == InternalVersion);
-
-	if(CurrentSupportedVersion == InternalVersion)
+	if (ensure(CurrentSupportedVersion == InternalVersion))
 	{		
 		// Load model
 		UnrealMutableInputStream Stream(Ar);
diff --git a/Engine/Plugins/Mutable/Source/MutableRuntime/Private/MuR/Image.cpp b/Engine/Plugins/Mutable/Source/MutableRuntime/Private/MuR/Image.cpp
--- a/Engine/Plugins/Mutable/Source/MutableRuntime/Private/MuR/Image.cpp
+++ b/Engine/Plugins/Mutable/Source/MutableRuntime/Private/MuR/Image.cpp
@@ -113,13 +113,9 @@ namespace mu
 				}
 			}
 		}
-		else if (Version >= 4)
-		{
-			Arch >> DataStorage;
-		}
 		else
 		{
-			check(false);
+			Arch >> DataStorage;
 		}
 
 		Arch >> m_flags;
diff --git a/Engine/Plugins/Online/VoiceChat/EOSVoiceChat/Source/EOSVoiceChat/Private/EOSVoiceChatUser.cpp b/Engine/Plugins/Online/VoiceChat/EOSVoiceChat/Source/EOSVoiceChat/Private/EOSVoiceChatUser.cpp
--- a/Engine/Plugins/Online/VoiceChat/EOSVoiceChat/Source/EOSVoiceChat/Private/EOSVoiceChatUser.cpp
+++ b/Engine/Plugins/Online/VoiceChat/EOSVoiceChat/Source/EOSVoiceChat/Private/EOSVoiceChatUser.cpp
@@ -1525,114 +1525,114 @@ void FEOSVoiceChatUser::LeaveChannelInternal(const FString& ChannelName, const F
 
 void FEOSVoiceChatUser::LogoutInternal(const FOnVoiceChatLogoutCompleteDelegate& Delegate)
 {
-	check(LoginSession.State == ELoginState::LoggedIn
-		|| LoginSession.State == ELoginState::LoggingOut);
-
-	if(LoginSession.State == ELoginState::LoggingOut)
+	if (ensure(LoginSession.State == ELoginState::LoggedIn || LoginSession.State == ELoginState::LoggingOut))
 	{
-		check(LoginSession.LogoutState.IsSet());
-		LoginSession.LogoutState->CompletionDelegates.Emplace(Delegate);
-	}
-	else if (LoginSession.State == ELoginState::LoggedIn)
-	{
-		LoginSession.State = ELoginState::LoggingOut;
-
-		// We have to actually leave all the channels, because there is no concept of "connected" or "logged in" in RTC, so no API calls we can call to "leave everything" like in vivox
-		TSet<FString> JoinedChannelNames;
-		TSet<FString> LeavingChannelNames;
-		TSet<FString> LobbyIds;
+		if(LoginSession.State == ELoginState::LoggingOut)
 		{
-			for (const TPair<FString, FChannelSession>& Pair : LoginSession.ChannelSessions)
+			check(LoginSession.LogoutState.IsSet());
+			LoginSession.LogoutState->CompletionDelegates.Emplace(Delegate);
+		}
+		else
+		{
+			LoginSession.State = ELoginState::LoggingOut;
+
+			// We have to actually leave all the channels, because there is no concept of "connected" or "logged in" in RTC, so no API calls we can call to "leave everything" like in vivox
+			TSet<FString> JoinedChannelNames;
+			TSet<FString> LeavingChannelNames;
+			TSet<FString> LobbyIds;
 			{
-				const FChannelSession& ChannelSession = Pair.Value;
-				if (ChannelSession.IsLobbySession())
-				{
-					LobbyIds.Add(ChannelSession.LobbyId);
-				}
-				else if (ChannelSession.JoinState >= EChannelJoinState::Joining)
+				for (const TPair<FString, FChannelSession>& Pair : LoginSession.ChannelSessions)
 				{
-					JoinedChannelNames.Add(ChannelSession.ChannelName);
-				}
-				else if (ChannelSession.JoinState == EChannelJoinState::Leaving)
-				{
-					LeavingChannelNames.Add(ChannelSession.ChannelName);
+					const FChannelSession& ChannelSession = Pair.Value;
+					if (ChannelSession.IsLobbySession())
+					{
+						LobbyIds.Add(ChannelSession.LobbyId);
+					}
+					else if (ChannelSession.JoinState >= EChannelJoinState::Joining)
+					{
+						JoinedChannelNames.Add(ChannelSession.ChannelName);
+					}
+					else if (ChannelSession.JoinState == EChannelJoinState::Leaving)
+					{
+						LeavingChannelNames.Add(ChannelSession.ChannelName);
+					}
 				}
 			}
-		}
-
-		for (const FString& LobbyId : LobbyIds)
-		{
-			RemoveLobbyRoom(LobbyId);
-		}
 
-		FVoiceChatResult Result = FVoiceChatResult::CreateSuccess();
-		if (JoinedChannelNames.Num() > 0 || LeavingChannelNames.Num() > 0)
-		{
-			// Channels are connected, so we fire leave requests in parallel, wait for them to complete, then return success/failure
-			FLoginSession::FLogoutState& LogoutState = LoginSession.LogoutState.Emplace();
-			LogoutState.CompletionDelegates.Emplace(Delegate);
-			LogoutState.ChannelNamesExpectingCallback.Reserve(JoinedChannelNames.Num() + LeavingChannelNames.Num());
-			LogoutState.ChannelNamesExpectingCallback.Append(JoinedChannelNames);
-			LogoutState.ChannelNamesExpectingCallback.Append(LeavingChannelNames);
+			for (const FString& LobbyId : LobbyIds)
+			{
+				RemoveLobbyRoom(LobbyId);
+			}
 
-			FOnVoiceChatChannelLeaveCompleteDelegate LeaveCompleteDelegate = FOnVoiceChatChannelLeaveCompleteDelegate::CreateLambda([this](const FString& ChannelName, const FVoiceChatResult& LambdaResult)
+			FVoiceChatResult Result = FVoiceChatResult::CreateSuccess();
+			if (JoinedChannelNames.Num() > 0 || LeavingChannelNames.Num() > 0)
 			{
-				if (!LambdaResult.IsSuccess())
+				// Channels are connected, so we fire leave requests in parallel, wait for them to complete, then return success/failure
+				FLoginSession::FLogoutState& LogoutState = LoginSession.LogoutState.Emplace();
+				LogoutState.CompletionDelegates.Emplace(Delegate);
+				LogoutState.ChannelNamesExpectingCallback.Reserve(JoinedChannelNames.Num() + LeavingChannelNames.Num());
+				LogoutState.ChannelNamesExpectingCallback.Append(JoinedChannelNames);
+				LogoutState.ChannelNamesExpectingCallback.Append(LeavingChannelNames);
+
+				FOnVoiceChatChannelLeaveCompleteDelegate LeaveCompleteDelegate = FOnVoiceChatChannelLeaveCompleteDelegate::CreateLambda([this](const FString& ChannelName, const FVoiceChatResult& LambdaResult)
 				{
-					EOSVOICECHATUSER_LOG(Warning, TEXT("LogoutInternal LeaveChannelInternal ChannelName=[%s] %s"), *ChannelName, *LexToString(LambdaResult));
-					LoginSession.LogoutState->Result = LambdaResult;
-				}
-				LoginSession.LogoutState->ChannelNamesExpectingCallback.Remove(ChannelName);
-				if (LoginSession.LogoutState->ChannelNamesExpectingCallback.Num() == 0)
-				{
-					const FVoiceChatResult Result = MoveTemp(LoginSession.LogoutState->Result);
-					const TArray<FOnVoiceChatLogoutCompleteDelegate> LogoutCompleteDelegates = MoveTemp(LoginSession.LogoutState->CompletionDelegates);
-					if (Result.IsSuccess())
+					if (!LambdaResult.IsSuccess())
 					{
-						// All the LeaveChannel requests succeeded, so clear the login session
-						const FString PlayerName = MoveTemp(LoginSession.PlayerName);
-						ClearLoginSession();
-
-						for (const FOnVoiceChatLogoutCompleteDelegate& Delegate : LogoutCompleteDelegates)
-						{
-							Delegate.ExecuteIfBound(PlayerName, FVoiceChatResult::CreateSuccess());
-						}
-						OnVoiceChatLoggedOutDelegate.Broadcast(PlayerName);
+						EOSVOICECHATUSER_LOG(Warning, TEXT("LogoutInternal LeaveChannelInternal ChannelName=[%s] %s"), *ChannelName, *LexToString(LambdaResult));
+						LoginSession.LogoutState->Result = LambdaResult;
 					}
-					else
+					LoginSession.LogoutState->ChannelNamesExpectingCallback.Remove(ChannelName);
+					if (LoginSession.LogoutState->ChannelNamesExpectingCallback.Num() == 0)
 					{
-						// One of the LeaveChannel requests failed, so we stay logged in
-						LoginSession.State = ELoginState::LoggedIn;
-						LoginSession.LogoutState.Reset();
-
-						for (const FOnVoiceChatLogoutCompleteDelegate& Delegate : LogoutCompleteDelegates)
+						const FVoiceChatResult Result = MoveTemp(LoginSession.LogoutState->Result);
+						const TArray<FOnVoiceChatLogoutCompleteDelegate> LogoutCompleteDelegates = MoveTemp(LoginSession.LogoutState->CompletionDelegates);
+						if (Result.IsSuccess())
+						{
+							// All the LeaveChannel requests succeeded, so clear the login session
+							const FString PlayerName = MoveTemp(LoginSession.PlayerName);
+							ClearLoginSession();
+
+							for (const FOnVoiceChatLogoutCompleteDelegate& Delegate : LogoutCompleteDelegates)
+							{
+								Delegate.ExecuteIfBound(PlayerName, FVoiceChatResult::CreateSuccess());
+							}
+							OnVoiceChatLoggedOutDelegate.Broadcast(PlayerName);
+						}
+						else
 						{
-							Delegate.ExecuteIfBound(LoginSession.PlayerName, Result);
+							// One of the LeaveChannel requests failed, so we stay logged in
+							LoginSession.State = ELoginState::LoggedIn;
+							LoginSession.LogoutState.Reset();
+
+							for (const FOnVoiceChatLogoutCompleteDelegate& Delegate : LogoutCompleteDelegates)
+							{
+								Delegate.ExecuteIfBound(LoginSession.PlayerName, Result);
+							}
 						}
 					}
+				});
+
+				for (const FString& ChannelName : JoinedChannelNames)
+				{
+					// Actually leave joined/joining channels
+					LeaveChannelInternal(ChannelName, LeaveCompleteDelegate);
 				}
-			});
 
-			for (const FString& ChannelName : JoinedChannelNames)
-			{
-				// Actually leave joined/joining channels
-				LeaveChannelInternal(ChannelName, LeaveCompleteDelegate);
+				for (const FString& ChannelName : LeavingChannelNames)
+				{
+					// Add to the LeaveDelegates of already leaving channels
+					GetChannelSession(ChannelName).LeaveDelegates.Add(LeaveCompleteDelegate);
+				}
 			}
-
-			for (const FString& ChannelName : LeavingChannelNames)
+			else
 			{
-				// Add to the LeaveDelegates of already leaving channels
-				GetChannelSession(ChannelName).LeaveDelegates.Add(LeaveCompleteDelegate);
-			}
-		}
-		else
-		{
-			// No channels connected, so just clear the LoginSession and fire successes
-			const FString PlayerName = MoveTemp(LoginSession.PlayerName);
-			ClearLoginSession();
+				// No channels connected, so just clear the LoginSession and fire successes
+				const FString PlayerName = MoveTemp(LoginSession.PlayerName);
+				ClearLoginSession();
 
-			Delegate.ExecuteIfBound(PlayerName, FVoiceChatResult::CreateSuccess());
-			OnVoiceChatLoggedOutDelegate.Broadcast(PlayerName);
+				Delegate.ExecuteIfBound(PlayerName, FVoiceChatResult::CreateSuccess());
+				OnVoiceChatLoggedOutDelegate.Broadcast(PlayerName);
+			}
 		}
 	}
 }
diff --git a/Engine/Plugins/Runtime/AudioCapture/Source/AudioCaptureEditor/Private/AudioRecordingManager.cpp b/Engine/Plugins/Runtime/AudioCapture/Source/AudioCaptureEditor/Private/AudioRecordingManager.cpp
--- a/Engine/Plugins/Runtime/AudioCapture/Source/AudioCaptureEditor/Private/AudioRecordingManager.cpp
+++ b/Engine/Plugins/Runtime/AudioCapture/Source/AudioCaptureEditor/Private/AudioRecordingManager.cpp
@@ -46,9 +46,8 @@ FAudioRecordingManager::~FAudioRecordingManager()
 void FAudioRecordingManager::StartRecording(const FRecordingSettings& InSettings)
 {
 	EAudioRecorderState CurrentState = RecorderState.load();
-	check(CurrentState == EAudioRecorderState::PreRecord);
 
-	if (CurrentState == EAudioRecorderState::PreRecord)
+	if (ensure(CurrentState == EAudioRecorderState::PreRecord))
 	{
 		RecorderState = EAudioRecorderState::Recording;
 
@@ -188,9 +187,8 @@ static void SampleRateConvert(float CurrentSR, float TargetSR, int32 NumChannels
 void FAudioRecordingManager::StopRecording()
 {
 	EAudioRecorderState CurrentState = RecorderState.load();
-	check(CurrentState == EAudioRecorderState::Recording);
 
-	if (CurrentState == EAudioRecorderState::Recording)
+	if (ensure(CurrentState == EAudioRecorderState::Recording))
 	{
 		RecorderState = EAudioRecorderState::Stopped;
 
@@ -219,9 +217,8 @@ void FAudioRecordingManager::StopRecording()
 TObjectPtr<USoundWave> FAudioRecordingManager::GetRecordedSoundWave(const FRecordingManagerSourceSettings& InSourceSettings)
 {
 	EAudioRecorderState CurrentState = RecorderState.load();
-	check(CurrentState == EAudioRecorderState::Stopped);
 
-	if (CurrentState == EAudioRecorderState::Stopped)
+	if (ensure(CurrentState == EAudioRecorderState::Stopped))
 	{
 		if (PCMDataToSerialize == nullptr)
 		{
diff --git a/Engine/Plugins/Runtime/MassGameplay/Source/MassActors/Private/MassAgentComponent.cpp b/Engine/Plugins/Runtime/MassGameplay/Source/MassActors/Private/MassAgentComponent.cpp
--- a/Engine/Plugins/Runtime/MassGameplay/Source/MassActors/Private/MassAgentComponent.cpp
+++ b/Engine/Plugins/Runtime/MassGameplay/Source/MassActors/Private/MassAgentComponent.cpp
@@ -17,8 +17,8 @@
 #include "GameFramework/CharacterMovementComponent.h"
 
 #define MASSAGENT_CHECK( condition, Format, ... ) \
-	checkf( condition, Format, ##__VA_ARGS__ ); \
-	UE_CVLOG(!(condition), GetOwner(), LogMass, Error, Format, ##__VA_ARGS__);
+	UE_CVLOG(!(condition), GetOwner(), LogMass, Error, Format, ##__VA_ARGS__); \
+	checkf( condition, Format, ##__VA_ARGS__ );
 
 //----------------------------------------------------------------------//
 // UMassAgentComponent
diff --git a/Engine/Source/Editor/LandscapeEditor/Private/LandscapeEditorDetailCustomization_ImportLayers.cpp b/Engine/Source/Editor/LandscapeEditor/Private/LandscapeEditorDetailCustomization_ImportLayers.cpp
--- a/Engine/Source/Editor/LandscapeEditor/Private/LandscapeEditorDetailCustomization_ImportLayers.cpp
+++ b/Engine/Source/Editor/LandscapeEditor/Private/LandscapeEditorDetailCustomization_ImportLayers.cpp
@@ -75,7 +75,7 @@ void FLandscapeEditorStructCustomization_FLandscapeImportLayer::CustomizeChildre
 	FName LayerName;
 	FText LayerNameText;
 	FPropertyAccess::Result Result = PropertyHandle_LayerName->GetValue(LayerName);
-	checkSlow(Result == FPropertyAccess::Success);
+	check(Result != FPropertyAccess::Fail);
 	LayerNameText = FText::FromName(LayerName);
 	if (Result == FPropertyAccess::MultipleValues)
 	{
diff --git a/Engine/Source/Editor/LandscapeEditor/Private/LandscapeEditorDetailCustomization_NewLandscape.cpp b/Engine/Source/Editor/LandscapeEditor/Private/LandscapeEditorDetailCustomization_NewLandscape.cpp
--- a/Engine/Source/Editor/LandscapeEditor/Private/LandscapeEditorDetailCustomization_NewLandscape.cpp
+++ b/Engine/Source/Editor/LandscapeEditor/Private/LandscapeEditorDetailCustomization_NewLandscape.cpp
@@ -712,7 +712,7 @@ FText FLandscapeEditorDetailCustomization_NewLandscape::GetSectionSize(TSharedRe
 {
 	int32 QuadsPerSection = 0;
 	FPropertyAccess::Result Result = PropertyHandle->GetValue(QuadsPerSection);
-	check(Result == FPropertyAccess::Success);
+	check(Result != FPropertyAccess::Fail);
 
 	if (Result == FPropertyAccess::MultipleValues)
 	{
@@ -747,7 +747,7 @@ FText FLandscapeEditorDetailCustomization_NewLandscape::GetSectionsPerComponent(
 {
 	int32 SectionsPerComponent = 0;
 	FPropertyAccess::Result Result = PropertyHandle->GetValue(SectionsPerComponent);
-	check(Result == FPropertyAccess::Success);
+	check(Result != FPropertyAccess::Success);
 
 	if (Result == FPropertyAccess::MultipleValues)
 	{
diff --git a/Engine/Source/Editor/LandscapeEditor/Private/LandscapeEditorDetailCustomization_ResizeLandscape.cpp b/Engine/Source/Editor/LandscapeEditor/Private/LandscapeEditorDetailCustomization_ResizeLandscape.cpp
--- a/Engine/Source/Editor/LandscapeEditor/Private/LandscapeEditorDetailCustomization_ResizeLandscape.cpp
+++ b/Engine/Source/Editor/LandscapeEditor/Private/LandscapeEditorDetailCustomization_ResizeLandscape.cpp
@@ -321,7 +321,7 @@ FText FLandscapeEditorDetailCustomization_ResizeLandscape::GetSectionSize(TShare
 {
 	int32 QuadsPerSection = 0;
 	FPropertyAccess::Result Result = PropertyHandle->GetValue(QuadsPerSection);
-	check(Result == FPropertyAccess::Success);
+	check(Result != FPropertyAccess::Fail);
 
 	if (Result == FPropertyAccess::MultipleValues)
 	{
@@ -391,7 +391,7 @@ FText FLandscapeEditorDetailCustomization_ResizeLandscape::GetSectionsPerCompone
 {
 	int32 SectionsPerComponent = 0;
 	FPropertyAccess::Result Result = PropertyHandle->GetValue(SectionsPerComponent);
-	check(Result == FPropertyAccess::Success);
+	check(Result != FPropertyAccess::Fail);
 
 	if (Result == FPropertyAccess::MultipleValues)
 	{
diff --git a/Engine/Source/Runtime/Core/Internal/Containers/DirectoryTree.h b/Engine/Source/Runtime/Core/Internal/Containers/DirectoryTree.h
--- a/Engine/Source/Runtime/Core/Internal/Containers/DirectoryTree.h
+++ b/Engine/Source/Runtime/Core/Internal/Containers/DirectoryTree.h
@@ -1269,7 +1269,7 @@ inline bool TDirectoryTree<ValueType>::FTreeNode::TryGetChildren(FStringBuilderB
 		checkf(false, TEXT("Infinite loop trying to split path %.*s into components."), InRelPath.Len(), InRelPath.GetData());
 	}
 
-	check(false); // Only way to get here is from the checkf in the else block.
+	check(false); //-V779 Only way to get here is from the checkf in the else block.
 	return false;
 }
 
diff --git a/Engine/Source/Runtime/Core/Private/ProfilingDebugging/CsvProfiler.cpp b/Engine/Source/Runtime/Core/Private/ProfilingDebugging/CsvProfiler.cpp
--- a/Engine/Source/Runtime/Core/Private/ProfilingDebugging/CsvProfiler.cpp
+++ b/Engine/Source/Runtime/Core/Private/ProfilingDebugging/CsvProfiler.cpp
@@ -692,8 +692,7 @@ public:
 		FScopeLock Lock(&CS);
 		{
 			Index = GetCategoryIndex(CategoryName);
-			checkf(Index == -1, TEXT("CSV stat category already declared: %s. Note: Categories are not case sensitive"), *CategoryName);
-			if (Index == -1)
+			if (ensureMsgf(Index == -1, TEXT("CSV stat category already declared: %s. Note: Categories are not case sensitive"), *CategoryName))
 			{
 				if (bIsGlobal)
 				{
diff --git a/Engine/Source/Runtime/Core/Private/Windows/WindowsPlatformNamedPipe.cpp b/Engine/Source/Runtime/Core/Private/Windows/WindowsPlatformNamedPipe.cpp
--- a/Engine/Source/Runtime/Core/Private/Windows/WindowsPlatformNamedPipe.cpp
+++ b/Engine/Source/Runtime/Core/Private/Windows/WindowsPlatformNamedPipe.cpp
@@ -245,7 +245,7 @@ bool FWindowsPlatformNamedPipe::BlockForAsyncIO()
 				break;
 		}
 	}
-	while (bTryAgain);
+	while (bTryAgain); //-V654
 
 	return true;
 }
diff --git a/Engine/Source/Runtime/Engine/Private/Materials/MaterialExpressionHLSL.cpp b/Engine/Source/Runtime/Engine/Private/Materials/MaterialExpressionHLSL.cpp
--- a/Engine/Source/Runtime/Engine/Private/Materials/MaterialExpressionHLSL.cpp
+++ b/Engine/Source/Runtime/Engine/Private/Materials/MaterialExpressionHLSL.cpp
@@ -1112,10 +1112,7 @@ bool UMaterialExpressionParticleSubUVProperties::GenerateHLSLExpression(FMateria
 
 	}
 	
-	if (Input != EExternalInput::None)
-	{
-		OutExpression = Generator.GetTree().NewExpression<FExpressionExternalInput>(Input);
-	}
+	OutExpression = Generator.GetTree().NewExpression<FExpressionExternalInput>(Input);
 	return OutExpression != nullptr;
 }
 
diff --git a/Engine/Source/Runtime/Engine/Private/ShaderCompiler/ShaderGenerationUtil.cpp b/Engine/Source/Runtime/Engine/Private/ShaderCompiler/ShaderGenerationUtil.cpp
--- a/Engine/Source/Runtime/Engine/Private/ShaderCompiler/ShaderGenerationUtil.cpp
+++ b/Engine/Source/Runtime/Engine/Private/ShaderCompiler/ShaderGenerationUtil.cpp
@@ -1526,7 +1526,7 @@ static FString CreateGBufferDecodeFunctionVariation(const FGBufferInfo& BufferIn
 					CoordName.GetCharArray().GetData(),
 					FullSwizzle[NumChan-1].GetCharArray().GetData());
 			}
-			else if (DecodeType == SceneTexturesLoad)
+			else if (DecodeType == SceneTexturesLoad) //-V547
 			{
 				CurrLine = FString::Printf(TEXT("\t%s InMRT%d = %sTexture.Load(int3(%s, 0)).%s;\n"),
 					TypeName.GetCharArray().GetData(),
diff --git a/Engine/Source/Runtime/Experimental/IoStore/HttpClient/Private/Loop.inl b/Engine/Source/Runtime/Experimental/IoStore/HttpClient/Private/Loop.inl
--- a/Engine/Source/Runtime/Experimental/IoStore/HttpClient/Private/Loop.inl
+++ b/Engine/Source/Runtime/Experimental/IoStore/HttpClient/Private/Loop.inl
@@ -554,7 +554,7 @@ static FOutcome DoRecv(FActivity* Activity, FPeerType& Peer, int32& MaxRecvSize)
 
 	if (State == EState::RecvMessage)	return DoRecvMessage(Activity, Peer);
 	if (State == EState::RecvContent)	return DoRecvContent(Activity, Peer, MaxRecvSize);
-	if (State == EState::RecvStream)	return DoRecvStream(Activity, Peer, MaxRecvSize);
+	if (State == EState::RecvStream)	return DoRecvStream(Activity, Peer, MaxRecvSize); //-V547
 	
 	check(false); // it is not expected that we'll get here
 	return FOutcome::Error("unreachable");
diff --git a/Engine/Source/Runtime/Foliage/Private/InstancedFoliage.cpp b/Engine/Source/Runtime/Foliage/Private/InstancedFoliage.cpp
--- a/Engine/Source/Runtime/Foliage/Private/InstancedFoliage.cpp
+++ b/Engine/Source/Runtime/Foliage/Private/InstancedFoliage.cpp
@@ -4599,8 +4599,7 @@ void AInstancedFoliageActor::PostLoad()
 
 			if (GetLinkerCustomVersion(FFoliageCustomVersion::GUID) < FFoliageCustomVersion::HierarchicalISMCNonTransactional)
 			{
-				check(Info.Type == EFoliageImplType::StaticMesh);
-				if (Info.Type == EFoliageImplType::StaticMesh)
+				if (ensure(Info.Type == EFoliageImplType::StaticMesh))
 				{
 					FFoliageStaticMesh* FoliageStaticMesh = StaticCast<FFoliageStaticMesh*>(Info.Implementation.Get());
 					if (FoliageStaticMesh->Component)
@@ -4622,8 +4621,7 @@ void AInstancedFoliageActor::PostLoad()
 			// Upgrade foliage component
 			if (GetLinkerCustomVersion(FFoliageCustomVersion::GUID) < FFoliageCustomVersion::FoliageUsingFoliageISMC)
 			{
-				check(Info.Type == EFoliageImplType::StaticMesh);
-				if (Info.Type == EFoliageImplType::StaticMesh)
+				if (ensure(Info.Type == EFoliageImplType::StaticMesh))
 				{
 					FFoliageStaticMesh* FoliageStaticMesh = StaticCast<FFoliageStaticMesh*>(Info.Implementation.Get());
 					UFoliageType_InstancedStaticMesh* FoliageType_InstancedStaticMesh = Cast<UFoliageType_InstancedStaticMesh>(FoliageType);
diff --git a/Engine/Source/Runtime/Serialization/Private/StructDeserializer.cpp b/Engine/Source/Runtime/Serialization/Private/StructDeserializer.cpp
--- a/Engine/Source/Runtime/Serialization/Private/StructDeserializer.cpp
+++ b/Engine/Source/Runtime/Serialization/Private/StructDeserializer.cpp
@@ -676,6 +676,7 @@ bool FStructDeserializer::DeserializeElement(void* OutAddress, UStruct& OwnerInf
 					}
 					else
 					{
+						checkSlow(CurrentState.Property);
 						//Too many entries in TArray
 						UE_LOG(LogSerialization, Verbose, TEXT("TArray %s has dimension of %d and trying to read element %d"), *CurrentState.Property->GetFName().ToString(), CurrentState.Property->ArrayDim, CurrentState.ArrayIndex);
 						continue;
@@ -683,6 +684,7 @@ bool FStructDeserializer::DeserializeElement(void* OutAddress, UStruct& OwnerInf
 				}
 				else
 				{
+					checkSlow(CurrentState.Property);
 					// static array element
 					if (CurrentState.ArrayIndex >= 0 && CurrentState.ArrayIndex < CurrentState.Property->ArrayDim)
 					{
@@ -1022,7 +1024,7 @@ bool FStructDeserializer::DeserializeElement(void* OutAddress, UStruct& OwnerInf
 						else
 						{
 							//Index out of bound
-							UE_LOG(LogSerialization, Verbose, TEXT("TMap %s has dimension of %d and trying to read element %d"), *CurrentState.Property->GetFName().ToString(), MapHelper.Num(), CurrentState.ArrayIndex);
+							UE_LOG(LogSerialization, Verbose, TEXT("TMap %s has dimension of %d and trying to read element %d"), *GetNameSafe(CurrentState.Property), MapHelper.Num(), CurrentState.ArrayIndex);
 							Backend.SkipStructure();
 							continue;
 						}
@@ -1048,7 +1050,7 @@ bool FStructDeserializer::DeserializeElement(void* OutAddress, UStruct& OwnerInf
 						else
 						{
 							//Index out of bound
-							UE_LOG(LogSerialization, Verbose, TEXT("TArray %s has dimension of %d and trying to read element %d"), *CurrentState.Property->GetFName().ToString(), SetHelper.Num(), CurrentState.ArrayIndex);
+							UE_LOG(LogSerialization, Verbose, TEXT("TArray %s has dimension of %d and trying to read element %d"), *GetNameSafe(CurrentState.Property), SetHelper.Num(), CurrentState.ArrayIndex);
 							Backend.SkipStructure();
 							continue;
 						}
@@ -1074,7 +1076,7 @@ bool FStructDeserializer::DeserializeElement(void* OutAddress, UStruct& OwnerInf
 						else
 						{
 							//Index out of bound
-							UE_LOG(LogSerialization, Verbose, TEXT("TArray %s has dimension of %d and trying to read element %d"), *CurrentState.Property->GetFName().ToString(), ArrayHelper.Num(), CurrentState.ArrayIndex);
+							UE_LOG(LogSerialization, Verbose, TEXT("TArray %s has dimension of %d and trying to read element %d"), *GetNameSafe(CurrentState.Property), ArrayHelper.Num(), CurrentState.ArrayIndex);
 							Backend.SkipStructure();
 							continue;
 						}
diff --git a/Engine/Source/Runtime/VulkanRHI/Private/VulkanTexture.cpp b/Engine/Source/Runtime/VulkanRHI/Private/VulkanTexture.cpp
--- a/Engine/Source/Runtime/VulkanRHI/Private/VulkanTexture.cpp
+++ b/Engine/Source/Runtime/VulkanRHI/Private/VulkanTexture.cpp
@@ -1678,7 +1678,7 @@ FVulkanTexture::FVulkanTexture(FRHICommandListBase* RHICmdList, FVulkanDevice& I
 
 	const VkDescriptorType DescriptorType = SupportsSampling() ? VK_DESCRIPTOR_TYPE_SAMPLED_IMAGE : VK_DESCRIPTOR_TYPE_STORAGE_IMAGE;
 	const VkImageUsageFlags SRVUsage = bIsSRGB ? (ImageCreateInfo.ImageCreateInfo.usage & ~VK_IMAGE_USAGE_STORAGE_BIT) : ImageCreateInfo.ImageCreateInfo.usage;
-	if (ViewType != VK_IMAGE_VIEW_TYPE_MAX_ENUM)
+	if (ViewType != VK_IMAGE_VIEW_TYPE_MAX_ENUM) //-V547
 	{
 		DefaultView = (new FVulkanView(InDevice, DescriptorType))->InitAsTextureView(
 			Image
@@ -2210,7 +2210,7 @@ void FVulkanTexture::UpdateLinkedViews()
 	const VkImageViewType ViewType = GetViewType();
 	const uint32 ArraySize = GetNumberOfArrayLevels();
 
-	if (ViewType != VK_IMAGE_VIEW_TYPE_MAX_ENUM)
+	if (ViewType != VK_IMAGE_VIEW_TYPE_MAX_ENUM) //-V547
 	{
 		DefaultView->InitAsTextureView(Image, ViewType, GetFullAspectMask(), GetDesc().Format, ViewFormat, 0, FMath::Max(NumMips, 1u), 0, ArraySize, !SupportsSampling());
 	}
diff --git a/Engine/Source/Runtime/Windows/D3D11RHI/Private/D3D11Util.cpp b/Engine/Source/Runtime/Windows/D3D11RHI/Private/D3D11Util.cpp
--- a/Engine/Source/Runtime/Windows/D3D11RHI/Private/D3D11Util.cpp
+++ b/Engine/Source/Runtime/Windows/D3D11RHI/Private/D3D11Util.cpp
@@ -129,7 +129,7 @@ static void TerminateOnDeviceRemoved(HRESULT D3DResult, ID3D11Device* Direct3DDe
 		// Workaround for the fact that in non-monolithic builds the exe gets into a weird state and exception handling fails. 
 		// @todo investigate why non-monolithic builds fail to capture the exception when graphics driver crashes.
 #if !IS_MONOLITHIC
-		FPlatformMisc::RequestExit(true, TEXT("TerminateOnDeviceRemoved"));
+		FPlatformMisc::RequestExit(true, TEXT("TerminateOnDeviceRemoved")); //-V779
 #endif
 	}
 }
