- Remove unused data APIs

#rnx

diff --git a/Engine/Plugins/Online/VoiceChat/EOSVoiceChat/Source/EOSVoiceChat/Private/EOSVoiceChat.cpp b/Engine/Plugins/Online/VoiceChat/EOSVoiceChat/Source/EOSVoiceChat/Private/EOSVoiceChat.cpp
--- a/Engine/Plugins/Online/VoiceChat/EOSVoiceChat/Source/EOSVoiceChat/Private/EOSVoiceChat.cpp
+++ b/Engine/Plugins/Online/VoiceChat/EOSVoiceChat/Source/EOSVoiceChat/Private/EOSVoiceChat.cpp
@@ -744,16 +744,6 @@ void FEOSVoiceChat::UnregisterOnVoiceChatBeforeRecvAudioRenderedDelegate(FDelega
 	GetVoiceChatUser().UnregisterOnVoiceChatBeforeRecvAudioRenderedDelegate(Handle);
 }
 
-FDelegateHandle FEOSVoiceChat::RegisterOnVoiceChatDataReceivedDelegate(const FOnVoiceChatDataReceivedDelegate::FDelegate& Delegate)
-{
-	return GetVoiceChatUser().RegisterOnVoiceChatDataReceivedDelegate(Delegate);
-}
-
-void FEOSVoiceChat::UnregisterOnVoiceChatDataReceivedDelegate(FDelegateHandle Handle)
-{
-	GetVoiceChatUser().UnregisterOnVoiceChatDataReceivedDelegate(Handle);
-}
-
 FString FEOSVoiceChat::InsecureGetLoginToken(const FString& PlayerName)
 {
 	return GetVoiceChatUser().InsecureGetLoginToken(PlayerName);
diff --git a/Engine/Plugins/Online/VoiceChat/EOSVoiceChat/Source/EOSVoiceChat/Private/EOSVoiceChatUser.cpp b/Engine/Plugins/Online/VoiceChat/EOSVoiceChat/Source/EOSVoiceChat/Private/EOSVoiceChatUser.cpp
--- a/Engine/Plugins/Online/VoiceChat/EOSVoiceChat/Source/EOSVoiceChat/Private/EOSVoiceChatUser.cpp
+++ b/Engine/Plugins/Online/VoiceChat/EOSVoiceChat/Source/EOSVoiceChat/Private/EOSVoiceChatUser.cpp
@@ -875,20 +875,6 @@ void FEOSVoiceChatUser::UnregisterOnVoiceChatBeforeRecvAudioRenderedDelegate(FDe
 	OnVoiceChatBeforeRecvAudioRenderedDelegate.Remove(Handle);
 }
 
-FDelegateHandle FEOSVoiceChatUser::RegisterOnVoiceChatDataReceivedDelegate(const FOnVoiceChatDataReceivedDelegate::FDelegate& Delegate)
-{
-	FScopeLock Lock(&DataReceivedLock);
-
-	return OnVoiceChatDataReceivedDelegate.Add(Delegate);
-}
-
-void FEOSVoiceChatUser::UnregisterOnVoiceChatDataReceivedDelegate(FDelegateHandle Handle)
-{
-	FScopeLock Lock(&DataReceivedLock);
-
-	OnVoiceChatDataReceivedDelegate.Remove(Handle);
-}
-
 FString FEOSVoiceChatUser::InsecureGetLoginToken(const FString& PlayerName)
 {
 	return FString();
diff --git a/Engine/Plugins/Online/VoiceChat/EOSVoiceChat/Source/EOSVoiceChat/Public/EOSVoiceChat.h b/Engine/Plugins/Online/VoiceChat/EOSVoiceChat/Source/EOSVoiceChat/Public/EOSVoiceChat.h
--- a/Engine/Plugins/Online/VoiceChat/EOSVoiceChat/Source/EOSVoiceChat/Public/EOSVoiceChat.h
+++ b/Engine/Plugins/Online/VoiceChat/EOSVoiceChat/Source/EOSVoiceChat/Public/EOSVoiceChat.h
@@ -38,8 +38,6 @@ public:
 	static FOnAudioStatusChanged OnAudioStatusChanged;
 };
 
-DECLARE_MULTICAST_DELEGATE_ThreeParams(FOnVoiceChatDataReceivedDelegate, const FString& /* ChannelName */, const FString& /* PlayerName */, TArrayView<const uint8> /* Buffer */);
-
 class FEOSAudioDevicePool;
 class FEOSVoiceChatUser;
 class IEOSSDKManager;
@@ -136,8 +134,6 @@ public:
 	virtual void UnregisterOnVoiceChatBeforeCaptureAudioSentDelegate(FDelegateHandle Handle) override;
 	virtual FDelegateHandle RegisterOnVoiceChatBeforeRecvAudioRenderedDelegate(const FOnVoiceChatBeforeRecvAudioRenderedDelegate::FDelegate& Delegate) override;
 	virtual void UnregisterOnVoiceChatBeforeRecvAudioRenderedDelegate(FDelegateHandle Handle) override;
-	virtual FDelegateHandle RegisterOnVoiceChatDataReceivedDelegate(const FOnVoiceChatDataReceivedDelegate::FDelegate& Delegate);
-	virtual void UnregisterOnVoiceChatDataReceivedDelegate(FDelegateHandle Handle);
 	virtual FString InsecureGetLoginToken(const FString& PlayerName) override;
 	virtual FString InsecureGetJoinToken(const FString& ChannelName, EVoiceChatChannelType ChannelType, TOptional<FVoiceChatChannel3dProperties> Channel3dProperties = TOptional<FVoiceChatChannel3dProperties>()) override;
 	// ~End IVoiceChatUser Interface
diff --git a/Engine/Plugins/Online/VoiceChat/EOSVoiceChat/Source/EOSVoiceChat/Public/EOSVoiceChatUser.h b/Engine/Plugins/Online/VoiceChat/EOSVoiceChat/Source/EOSVoiceChat/Public/EOSVoiceChatUser.h
--- a/Engine/Plugins/Online/VoiceChat/EOSVoiceChat/Source/EOSVoiceChat/Public/EOSVoiceChatUser.h
+++ b/Engine/Plugins/Online/VoiceChat/EOSVoiceChat/Source/EOSVoiceChat/Public/EOSVoiceChatUser.h
@@ -82,8 +82,6 @@ public:
 	virtual void UnregisterOnVoiceChatBeforeCaptureAudioSentDelegate(FDelegateHandle Handle) override;
 	virtual FDelegateHandle RegisterOnVoiceChatBeforeRecvAudioRenderedDelegate(const FOnVoiceChatBeforeRecvAudioRenderedDelegate::FDelegate& Delegate) override;
 	virtual void UnregisterOnVoiceChatBeforeRecvAudioRenderedDelegate(FDelegateHandle Handle) override;
-	virtual FDelegateHandle RegisterOnVoiceChatDataReceivedDelegate(const FOnVoiceChatDataReceivedDelegate::FDelegate& Delegate);
-	virtual void UnregisterOnVoiceChatDataReceivedDelegate(FDelegateHandle Handle);
 	virtual FString InsecureGetLoginToken(const FString& PlayerName) override;
 	virtual FString InsecureGetJoinToken(const FString& ChannelName, EVoiceChatChannelType ChannelType, TOptional<FVoiceChatChannel3dProperties> Channel3dProperties = TOptional<FVoiceChatChannel3dProperties>()) override;
 	// ~End IVoiceChatUser Interface
@@ -290,8 +288,6 @@ protected:
 	FOnVoiceChatBeforeCaptureAudioSentDelegate OnVoiceChatBeforeCaptureAudioSentDelegate;
 	FCriticalSection BeforeRecvAudioRenderedLock;
 	FOnVoiceChatBeforeRecvAudioRenderedDelegate OnVoiceChatBeforeRecvAudioRenderedDelegate;
-	FCriticalSection DataReceivedLock;
-	FOnVoiceChatDataReceivedDelegate OnVoiceChatDataReceivedDelegate;
 
 	FEOSVoiceChat& EOSVoiceChat;
 
