- Reduce log spam from EOSVoiceChat when joining heavily populated channels

diff --git a/Engine/Plugins/Online/VoiceChat/EOSVoiceChat/Source/EOSVoiceChat/Private/EOSVoiceChatUser.cpp b/Engine/Plugins/Online/VoiceChat/EOSVoiceChat/Source/EOSVoiceChat/Private/EOSVoiceChatUser.cpp
--- a/Engine/Plugins/Online/VoiceChat/EOSVoiceChat/Source/EOSVoiceChat/Private/EOSVoiceChatUser.cpp
+++ b/Engine/Plugins/Online/VoiceChat/EOSVoiceChat/Source/EOSVoiceChat/Private/EOSVoiceChatUser.cpp
@@ -1854,7 +1854,7 @@ void FEOSVoiceChatUser::OnUpdateParticipantVolume(const EOS_RTCAudio_UpdateParti
 
 	if (ResultCode == EOS_EResult::EOS_Success)
 	{
-		EOSVOICECHATUSER_LOG(Log, TEXT("OnUpdateParticipantVolume ChannelName=[%s] PlayerName=[%s] Success"), *ChannelName, *PlayerName);
+		EOSVOICECHATUSER_LOG(VeryVerbose, TEXT("OnUpdateParticipantVolume ChannelName=[%s] PlayerName=[%s] Success"), *ChannelName, *PlayerName);
 		FChannelSession& ChannelSession = GetChannelSession(ChannelName);
 		if (FChannelParticipant* ChannelParticipant = ChannelSession.Participants.Find(PlayerName))
 		{
@@ -1901,7 +1901,7 @@ void FEOSVoiceChatUser::OnUpdateReceivingAudio(const EOS_RTCAudio_UpdateReceivin
 
 	if (ResultCode == EOS_EResult::EOS_Success)
 	{
-		EOSVOICECHATUSER_LOG(Log, TEXT("OnUpdateReceiving ChannelName=[%s] PlayerName=[%s] Success"), *ChannelName, *PlayerName);
+		EOSVOICECHATUSER_LOG(VeryVerbose, TEXT("OnUpdateReceiving ChannelName=[%s] PlayerName=[%s] Success"), *ChannelName, *PlayerName);
 
 		// UpdateReceiving is called once prior to JoinRoom to set the default receiving state, in which case PlayerName is empty
 		if (!PlayerName.IsEmpty())
@@ -1931,7 +1931,7 @@ void FEOSVoiceChatUser::OnUpdateSendingAudio(const EOS_RTCAudio_UpdateSendingCal
 
 	if (CallbackInfo->ResultCode == EOS_EResult::EOS_Success)
 	{
-		EOSVOICECHATUSER_LOG(Log, TEXT("OnUpdateSending ChannelName=[%s] Success"), *ChannelName);
+		EOSVOICECHATUSER_LOG(VeryVerbose, TEXT("OnUpdateSending ChannelName=[%s] Success"), *ChannelName);
 
 		if (IsLoggedIn())
 		{
@@ -2099,7 +2099,7 @@ void FEOSVoiceChatUser::OnChannelParticipantStatusChanged(const EOS_RTC_Particip
 		{
 			if (PlayerName != ChannelSession->PlayerName)
 			{
-				EOSVOICECHATUSER_LOG(Log, TEXT("OnChannelParticipantJoined ChannelName=[%s] PlayerName=[%s]"), *ChannelName, *PlayerName);
+				EOSVOICECHATUSER_LOG(Verbose, TEXT("OnChannelParticipantJoined ChannelName=[%s] PlayerName=[%s]"), *ChannelName, *PlayerName);
 
 				const FGlobalParticipant& GlobalParticipant = GetGlobalParticipant(PlayerName);
 				FChannelParticipant& ChannelParticipant = ChannelSession->Participants.Add(PlayerName);
@@ -2127,7 +2127,7 @@ void FEOSVoiceChatUser::OnChannelParticipantStatusChanged(const EOS_RTC_Particip
 	}
 	else
 	{
-		EOSVOICECHATUSER_LOG(Log, TEXT("OnChannelParticipantLeft ChannelName=[%s] PlayerName=[%s]"), *ChannelName, *PlayerName);
+		EOSVOICECHATUSER_LOG(Verbose, TEXT("OnChannelParticipantLeft ChannelName=[%s] PlayerName=[%s]"), *ChannelName, *PlayerName);
 
 		if (FChannelSession* ChannelSession = LoginSession.ChannelSessions.Find(ChannelName))
 		{
