- Call EOS_RTC_SetSetting to set MaxTickInterval, reading the value from config

#rnx
#Synced-CL 48600773

diff --git a/Engine/Plugins/Online/VoiceChat/EOSVoiceChat/Source/EOSVoiceChat/Private/EOSVoiceChat.cpp b/Engine/Plugins/Online/VoiceChat/EOSVoiceChat/Source/EOSVoiceChat/Private/EOSVoiceChat.cpp
--- a/Engine/Plugins/Online/VoiceChat/EOSVoiceChat/Source/EOSVoiceChat/Private/EOSVoiceChat.cpp
+++ b/Engine/Plugins/Online/VoiceChat/EOSVoiceChat/Source/EOSVoiceChat/Private/EOSVoiceChat.cpp
@@ -7,6 +7,7 @@
 #include "Async/Async.h"
 #include "Containers/Ticker.h"
 #include "Misc/ConfigCacheIni.h"
+#include "Misc/CoreDelegates.h"
 #include "Misc/ScopeLock.h"
 #include "ProfilingDebugging/CsvProfiler.h"
 #include "Stats/Stats.h"
@@ -768,6 +769,9 @@ void FEOSVoiceChat::BindInitCallbacks()
 	}
 
 	OnAudioDevicesChanged();
+
+	FCoreDelegates::TSOnConfigSectionsChanged().AddThreadSafeSP(this, &FEOSVoiceChat::OnConfigSectionsChanged);
+	LoadConfig();
 }
 
 void FEOSVoiceChat::UnbindInitCallbacks()
@@ -777,6 +781,8 @@ void FEOSVoiceChat::UnbindInitCallbacks()
 		EOS_RTCAudio_RemoveNotifyAudioDevicesChanged(EOS_RTC_GetAudioInterface(GetRtcInterface()), InitSession.OnAudioDevicesChangedNotificationId);
 		InitSession.OnAudioDevicesChangedNotificationId = EOS_INVALID_NOTIFICATIONID;
 	}
+
+	FCoreDelegates::TSOnConfigSectionsChanged().RemoveAll(this);
 }
 
 void FEOSVoiceChat::OnAudioDevicesChangedStatic(const EOS_RTCAudio_AudioDevicesChangedCallbackInfo* CallbackInfo)
@@ -813,6 +819,39 @@ void FEOSVoiceChat::OnAudioDevicesChanged()
 	}));
 }
 
+void FEOSVoiceChat::OnConfigSectionsChanged(const FString& IniFilename, const TSet<FString>& SectionNames)
+{
+	if (IniFilename == GEngineIni)
+	{
+		LoadConfig();
+	}
+}
+
+void FEOSVoiceChat::LoadConfig()
+{
+	auto SetFromConfig = [this](const TCHAR* InSettingName)
+	{
+		FString ConfigKey = FString(TEXT("RTCSetting_")) + InSettingName;
+		FString ConfigValue;
+		if (GConfig->GetString(TEXT("EOSVoiceChat"), *ConfigKey, ConfigValue, GEngineIni))
+		{
+			const FTCHARToUTF8 Utf8Key(InSettingName);
+			const FTCHARToUTF8 Utf8Value(*ConfigValue);
+
+			EOS_RTC_SetSettingOptions SetSettingOptions = {};
+			SetSettingOptions.ApiVersion = EOS_RTC_SETSETTING_API_LATEST;
+			static_assert(EOS_RTC_SETSETTING_API_LATEST == 1, "EOS_RTC_SetSettingOptions updated, check new fields");
+			SetSettingOptions.SettingName = Utf8Key.Get();
+			SetSettingOptions.SettingValue = Utf8Value.Get();
+			const EOS_EResult Result = EOS_RTC_SetSetting(GetRtcInterface(), &SetSettingOptions);
+			UE_CLOG(Result != EOS_EResult::EOS_Success, LogEOSVoiceChat, Warning, TEXT("SetSetting Key=%s Value=%s Result=%s"),
+				InSettingName, *ConfigValue, *LexToString(Result));
+		}
+	};
+
+	SetFromConfig(TEXT("MaxTickInterval"));
+}
+
 void FEOSVoiceChat::ScheduleReleaseUser(IVoiceChatUser* User)
 {
 	FEOSVoiceChatUserRef UserRef = static_cast<FEOSVoiceChatUser*>(User)->AsShared();
diff --git a/Engine/Plugins/Online/VoiceChat/EOSVoiceChat/Source/EOSVoiceChat/Public/EOSVoiceChat.h b/Engine/Plugins/Online/VoiceChat/EOSVoiceChat/Source/EOSVoiceChat/Public/EOSVoiceChat.h
--- a/Engine/Plugins/Online/VoiceChat/EOSVoiceChat/Source/EOSVoiceChat/Public/EOSVoiceChat.h
+++ b/Engine/Plugins/Online/VoiceChat/EOSVoiceChat/Source/EOSVoiceChat/Public/EOSVoiceChat.h
@@ -186,6 +186,9 @@ protected:
 	static void EOS_CALL OnAudioDevicesChangedStatic(const EOS_RTCAudio_AudioDevicesChangedCallbackInfo* CallbackInfo);
 	void OnAudioDevicesChanged();
 
+	void OnConfigSectionsChanged(const FString& IniFilename, const TSet<FString>& SectionNames);
+	void LoadConfig();
+
 	// The current state of the connection.
 	enum class EConnectionState
 	{
