UE : EOSVoiceChat: Support named config mechanism

#Synced-CL 50301277

diff --git a/Engine/Plugins/Online/OnlineSubsystemEOS/Source/OnlineSubsystemEOS/Private/OnlineSubsystemEOS.cpp b/Engine/Plugins/Online/OnlineSubsystemEOS/Source/OnlineSubsystemEOS/Private/OnlineSubsystemEOS.cpp
--- a/Engine/Plugins/Online/OnlineSubsystemEOS/Source/OnlineSubsystemEOS/Private/OnlineSubsystemEOS.cpp
+++ b/Engine/Plugins/Online/OnlineSubsystemEOS/Source/OnlineSubsystemEOS/Private/OnlineSubsystemEOS.cpp
@@ -801,6 +801,30 @@ IVoiceChatUser* FOnlineSubsystemEOS::GetVoiceChatUserInterface(const FUniqueNetI
 		if (FEOSVoiceChatFactory* EOSVoiceChatFactory = FEOSVoiceChatFactory::Get())
 		{
 			VoiceChatInterface = EOSVoiceChatFactory->CreateInstanceWithPlatform(EOSPlatformHandle);
+
+			if (VoiceChatInterface)
+			{
+				if (VoiceChatInterface->Initialize())
+				{
+					TOptional<bool> bConnectResult;
+					VoiceChatInterface->Connect(FOnVoiceChatConnectCompleteDelegate::CreateLambda(
+					[&bConnectResult](const FVoiceChatResult& VoiceChatResult)
+					{
+						bConnectResult = VoiceChatResult.IsSuccess();
+					}));
+
+					if (!bConnectResult.IsSet() || !bConnectResult.GetValue())
+					{
+						UE_LOG_ONLINE(Warning, TEXT("GetVoiceChatUserInterface: VoiceChat Connect failed"));
+						VoiceChatInterface.Reset();
+					}
+				}
+				else
+				{
+					UE_LOG_ONLINE(Warning, TEXT("GetVoiceChatUserInterface: VoiceChat Initialize failed"));
+					VoiceChatInterface.Reset();
+				}
+			}
 		}
 	}
 
diff --git a/Engine/Plugins/Online/VoiceChat/EOSVoiceChat/Source/EOSVoiceChat/Private/Android/AndroidEOSVoiceChat.cpp b/Engine/Plugins/Online/VoiceChat/EOSVoiceChat/Source/EOSVoiceChat/Private/Android/AndroidEOSVoiceChat.cpp
--- a/Engine/Plugins/Online/VoiceChat/EOSVoiceChat/Source/EOSVoiceChat/Private/Android/AndroidEOSVoiceChat.cpp
+++ b/Engine/Plugins/Online/VoiceChat/EOSVoiceChat/Source/EOSVoiceChat/Private/Android/AndroidEOSVoiceChat.cpp
@@ -6,8 +6,8 @@
 
 #include "AndroidEOSVoiceChatUser.h"
 
-FAndroidEOSVoiceChat::FAndroidEOSVoiceChat(IEOSSDKManager& InSDKManager, const IEOSPlatformHandlePtr& InPlatformHandle)
-	: FEOSVoiceChat(InSDKManager, InPlatformHandle)
+FAndroidEOSVoiceChat::FAndroidEOSVoiceChat(IEOSSDKManager& InSDKManager)
+	: FEOSVoiceChat(InSDKManager)
 {
 }
 
diff --git a/Engine/Plugins/Online/VoiceChat/EOSVoiceChat/Source/EOSVoiceChat/Private/Android/AndroidEOSVoiceChat.h b/Engine/Plugins/Online/VoiceChat/EOSVoiceChat/Source/EOSVoiceChat/Private/Android/AndroidEOSVoiceChat.h
--- a/Engine/Plugins/Online/VoiceChat/EOSVoiceChat/Source/EOSVoiceChat/Private/Android/AndroidEOSVoiceChat.h
+++ b/Engine/Plugins/Online/VoiceChat/EOSVoiceChat/Source/EOSVoiceChat/Private/Android/AndroidEOSVoiceChat.h
@@ -9,7 +9,7 @@
 class FAndroidEOSVoiceChat : public FEOSVoiceChat
 {
 public:
-	FAndroidEOSVoiceChat(IEOSSDKManager& InSDKManager, const IEOSPlatformHandlePtr& InPlatformHandle);
+	FAndroidEOSVoiceChat(IEOSSDKManager& InSDKManager);
 	virtual ~FAndroidEOSVoiceChat() = default;
 
 	// ~Begin IVoiceChat Interface 
diff --git a/Engine/Plugins/Online/VoiceChat/EOSVoiceChat/Source/EOSVoiceChat/Private/EOSVoiceChat.cpp b/Engine/Plugins/Online/VoiceChat/EOSVoiceChat/Source/EOSVoiceChat/Private/EOSVoiceChat.cpp
--- a/Engine/Plugins/Online/VoiceChat/EOSVoiceChat/Source/EOSVoiceChat/Private/EOSVoiceChat.cpp
+++ b/Engine/Plugins/Online/VoiceChat/EOSVoiceChat/Source/EOSVoiceChat/Private/EOSVoiceChat.cpp
@@ -25,6 +25,8 @@
 #include "eos_rtc.h"
 #include "eos_rtc_audio.h"
 
+#define CONFIG_SECTION_NAME TEXT("EOSVoiceChat")
+#define CONFIG_FILE GEngineIni
 #define CHECKPIN() FEOSVoiceChatPtr StrongThis = WeakThis.Pin(); if(!StrongThis) return
 
 DEFINE_LOG_CATEGORY(LogEOSVoiceChat);
@@ -37,9 +39,42 @@ int64 FEOSVoiceChat::StaticInstanceIdCount = 0;
 
 #define EOS_VOICE_TODO 0
 
+// Added for tests
+#if !UE_BUILD_SHIPPING
+namespace
+{
+enum class EEOSVoiceInitializationPath : uint8
+{
+	None,
+	PlatformHandleInjected,  // SetPlatformHandle
+	NamedConfig,             // SDKManager.CreatePlatform(ConfigName)
+	LegacyInlineConfig       // EOSPlatformCreate with inline credentials
+};
+
+const TCHAR* LexToString(EEOSVoiceInitializationPath Path)
+{
+	switch (Path)
+	{
+		default: case EEOSVoiceInitializationPath::None:			return TEXT("None");
+		case EEOSVoiceInitializationPath::PlatformHandleInjected:	return TEXT("PlatformHandleInjected");
+		case EEOSVoiceInitializationPath::NamedConfig:				return TEXT("NamedConfig");
+		case EEOSVoiceInitializationPath::LegacyInlineConfig:		return TEXT("LegacyInlineConfig");
+	}
+}
+}
+#endif // !UE_BUILD_SHIPPING
+
 FEOSVoiceChat::FEOSVoiceChat(IEOSSDKManager& InSDKManager, const IEOSPlatformHandlePtr& InPlatformHandle)
 	: SDKManager(InSDKManager)
-	, EosPlatformHandle(InPlatformHandle)
+{
+	SetPlatformHandle(InPlatformHandle);
+#if !UE_BUILD_SHIPPING
+	InitializationPath = static_cast<uint8>(EEOSVoiceInitializationPath::None);
+#endif // !UE_BUILD_SHIPPING
+}
+
+FEOSVoiceChat::FEOSVoiceChat(IEOSSDKManager& InSDKManager)
+	: SDKManager(InSDKManager)
 {
 }
 
@@ -83,78 +118,110 @@ void FEOSVoiceChat::Initialize(const FOnVoiceChatInitializeCompleteDelegate& Ini
 	case EInitializationState::Uninitialized:
 	{
 		bool bEnabled = true;
-		GConfig->GetBool(TEXT("EOSVoiceChat"), TEXT("bEnabled"), bEnabled, GEngineIni);
+		GConfig->GetBool(CONFIG_SECTION_NAME, TEXT("bEnabled"), bEnabled, CONFIG_FILE);
+		FString PlatformConfigName;
+		if (OverridePlatformConfigName.IsSet())
+		{
+			PlatformConfigName = OverridePlatformConfigName.GetValue();
+		}
+		else if (!GConfig->GetString(CONFIG_SECTION_NAME, TEXT("PlatformConfigName"), PlatformConfigName, CONFIG_FILE) || PlatformConfigName.IsEmpty())
+		{
+			PlatformConfigName = SDKManager.GetDefaultPlatformConfigName();
+		}
+
 		if (bEnabled)
 		{
 			InitSession.State = EInitializationState::Initializing;
 
+#if !UE_BUILD_SHIPPING
+			if (EosPlatformHandle)
+			{
+				InitializationPath = static_cast<uint8>(EEOSVoiceInitializationPath::PlatformHandleInjected);
+			}
+#endif // !UE_BUILD_SHIPPING
+
 			if (!EosPlatformHandle)
 			{
 				if (SDKManager.IsInitialized())
 				{
-					FString ConfigProductId;
-					FString ConfigSandboxId;
-					FString ConfigDeploymentId;
-					FString ConfigClientId;
-					FString ConfigClientSecret;
-					FString ConfigEncryptionKey;
-					FString ConfigOverrideCountryCode;
-					FString ConfigOverrideLocaleCode;
-
-					GConfig->GetString(TEXT("EOSVoiceChat"), TEXT("ProductId"), ConfigProductId, GEngineIni);
-					GConfig->GetString(TEXT("EOSVoiceChat"), TEXT("SandboxId"), ConfigSandboxId, GEngineIni);
-					GConfig->GetString(TEXT("EOSVoiceChat"), TEXT("DeploymentId"), ConfigDeploymentId, GEngineIni);
-					GConfig->GetString(TEXT("EOSVoiceChat"), TEXT("ClientId"), ConfigClientId, GEngineIni);
-					GConfig->GetString(TEXT("EOSVoiceChat"), TEXT("ClientSecret"), ConfigClientSecret, GEngineIni);
-					GConfig->GetString(TEXT("EOSVoiceChat"), TEXT("ClientEncryptionKey"), ConfigEncryptionKey, GEngineIni);
-					GConfig->GetString(TEXT("EOSVoiceChat"), TEXT("OverrideCountryCode"), ConfigOverrideCountryCode, GEngineIni);
-					GConfig->GetString(TEXT("EOSVoiceChat"), TEXT("OverrideLocaleCode"), ConfigOverrideLocaleCode, GEngineIni);
-
-					const FTCHARToUTF8 Utf8ProductId(*ConfigProductId);
-					const FTCHARToUTF8 Utf8SandboxId(*ConfigSandboxId);
-					const FTCHARToUTF8 Utf8DeploymentId(*ConfigDeploymentId);
-					const FTCHARToUTF8 Utf8ClientId(*ConfigClientId);
-					const FTCHARToUTF8 Utf8ClientSecret(*ConfigClientSecret);
-					const FTCHARToUTF8 Utf8EncryptionKey(*ConfigEncryptionKey);
-					const FTCHARToUTF8 Utf8OverrideCountryCode(*ConfigOverrideCountryCode);
-					const FTCHARToUTF8 Utf8OverrideLocaleCode(*ConfigOverrideLocaleCode);
-
-					EOS_Platform_Options PlatformOptions = {};
-					PlatformOptions.ApiVersion = 14;
-					UE_EOS_CHECK_API_MISMATCH(EOS_PLATFORM_OPTIONS_API_LATEST, 14);
-					PlatformOptions.Reserved = nullptr;
-					PlatformOptions.SystemSpecificOptions = nullptr;
-					PlatformOptions.ProductId = ConfigProductId.IsEmpty() ? nullptr : Utf8ProductId.Get();
-					PlatformOptions.SandboxId = ConfigSandboxId.IsEmpty() ? nullptr : Utf8SandboxId.Get();
-					PlatformOptions.ClientCredentials.ClientId = ConfigClientId.IsEmpty() ? nullptr : Utf8ClientId.Get();
-					PlatformOptions.ClientCredentials.ClientSecret = ConfigClientSecret.IsEmpty() ? nullptr : Utf8ClientSecret.Get();
-					PlatformOptions.bIsServer = false;
-					PlatformOptions.EncryptionKey = ConfigEncryptionKey.IsEmpty() ? nullptr : Utf8EncryptionKey.Get();
-					PlatformOptions.OverrideCountryCode = ConfigOverrideCountryCode.IsEmpty() ? nullptr : Utf8OverrideCountryCode.Get();
-					PlatformOptions.OverrideLocaleCode = ConfigOverrideLocaleCode.IsEmpty() ? nullptr : Utf8OverrideLocaleCode.Get();
-					PlatformOptions.DeploymentId = ConfigDeploymentId.IsEmpty() ? nullptr : Utf8DeploymentId.Get();
-					PlatformOptions.Flags = EOS_PF_DISABLE_OVERLAY;
-					PlatformOptions.CacheDirectory = nullptr;
-					PlatformOptions.TickBudgetInMilliseconds = 1;
-					PlatformOptions.IntegratedPlatformOptionsContainerHandle = nullptr;
-					PlatformOptions.TaskNetworkTimeoutSeconds = nullptr;
+					if (!PlatformConfigName.IsEmpty())
+					{
+						// UE-193389: now the named config mechanism in SDKManager loads all credentials from its own config
+						EosPlatformHandle = SDKManager.CreatePlatform(PlatformConfigName);
+
+#if !UE_BUILD_SHIPPING
+						InitializationPath = static_cast<uint8>(EEOSVoiceInitializationPath::NamedConfig);
+#endif // !UE_BUILD_SHIPPING
+
+					}
+					else
+					{
+						// Legacy path: read credentials inline and build platform options
+						FString ConfigProductId;
+						FString ConfigSandboxId;
+						FString ConfigDeploymentId;
+						FString ConfigClientId;
+						FString ConfigClientSecret;
+						FString ConfigEncryptionKey;
+						FString ConfigOverrideCountryCode;
+						FString ConfigOverrideLocaleCode;
+						GConfig->GetString(CONFIG_SECTION_NAME, TEXT("ProductId"), ConfigProductId, CONFIG_FILE);
+						GConfig->GetString(CONFIG_SECTION_NAME, TEXT("SandboxId"), ConfigSandboxId, CONFIG_FILE);
+						GConfig->GetString(CONFIG_SECTION_NAME, TEXT("DeploymentId"), ConfigDeploymentId, CONFIG_FILE);
+						GConfig->GetString(CONFIG_SECTION_NAME, TEXT("ClientId"), ConfigClientId, CONFIG_FILE);
+						GConfig->GetString(CONFIG_SECTION_NAME, TEXT("ClientSecret"), ConfigClientSecret, CONFIG_FILE);
+						GConfig->GetString(CONFIG_SECTION_NAME, TEXT("ClientEncryptionKey"), ConfigEncryptionKey, CONFIG_FILE);
+						GConfig->GetString(CONFIG_SECTION_NAME, TEXT("OverrideCountryCode"), ConfigOverrideCountryCode, CONFIG_FILE);
+						GConfig->GetString(CONFIG_SECTION_NAME, TEXT("OverrideLocaleCode"), ConfigOverrideLocaleCode, CONFIG_FILE);
+						const FTCHARToUTF8 Utf8ProductId(*ConfigProductId);
+						const FTCHARToUTF8 Utf8SandboxId(*ConfigSandboxId);
+						const FTCHARToUTF8 Utf8DeploymentId(*ConfigDeploymentId);
+						const FTCHARToUTF8 Utf8ClientId(*ConfigClientId);
+						const FTCHARToUTF8 Utf8ClientSecret(*ConfigClientSecret);
+						const FTCHARToUTF8 Utf8EncryptionKey(*ConfigEncryptionKey);
+						const FTCHARToUTF8 Utf8OverrideCountryCode(*ConfigOverrideCountryCode);
+						const FTCHARToUTF8 Utf8OverrideLocaleCode(*ConfigOverrideLocaleCode);
+						EOS_Platform_Options PlatformOptions = { };
+						PlatformOptions.ApiVersion = 14;
+						UE_EOS_CHECK_API_MISMATCH(EOS_PLATFORM_OPTIONS_API_LATEST, 14);
+						PlatformOptions.Reserved = nullptr;
+						PlatformOptions.SystemSpecificOptions = nullptr;
+						PlatformOptions.ProductId = ConfigProductId.IsEmpty() ? nullptr : Utf8ProductId.Get();
+						PlatformOptions.SandboxId = ConfigSandboxId.IsEmpty() ? nullptr : Utf8SandboxId.Get();
+						PlatformOptions.ClientCredentials.ClientId = ConfigClientId.IsEmpty() ? nullptr : Utf8ClientId.Get();
+						PlatformOptions.ClientCredentials.ClientSecret = ConfigClientSecret.IsEmpty() ? nullptr : Utf8ClientSecret.Get();
+						PlatformOptions.bIsServer = false;
+						PlatformOptions.EncryptionKey = ConfigEncryptionKey.IsEmpty() ? nullptr : Utf8EncryptionKey.Get();
+						PlatformOptions.OverrideCountryCode = ConfigOverrideCountryCode.IsEmpty() ? nullptr : Utf8OverrideCountryCode.Get();
+						PlatformOptions.OverrideLocaleCode = ConfigOverrideLocaleCode.IsEmpty() ? nullptr : Utf8OverrideLocaleCode.Get();
+						PlatformOptions.DeploymentId = ConfigDeploymentId.IsEmpty() ? nullptr : Utf8DeploymentId.Get();
+						PlatformOptions.Flags = EOS_PF_DISABLE_OVERLAY;
+						PlatformOptions.CacheDirectory = nullptr;
+						PlatformOptions.TickBudgetInMilliseconds = 1;
+						PlatformOptions.IntegratedPlatformOptionsContainerHandle = nullptr;
+						PlatformOptions.TaskNetworkTimeoutSeconds = nullptr;
 #if UE_EDITOR
-					//PlatformCreateOptions.Flags |= EOS_PF_LOADING_IN_EDITOR;
+						//PlatformCreateOptions.Flags |= EOS_PF_LOADING_IN_EDITOR;
 #endif
+						EOS_Platform_RTCOptions PlatformRTCOptions = {};
+						PlatformRTCOptions.ApiVersion = 3;
+						UE_EOS_CHECK_API_MISMATCH(EOS_PLATFORM_RTCOPTIONS_API_LATEST, 3);
+						PlatformRTCOptions.Reserved = nullptr;
+						PlatformRTCOptions.BackgroundMode = EOS_ERTCBackgroundMode::EOS_RTCBM_KeepRoomsAlive;
+						FString RTCBackgroundModeStr;
+						if(GConfig->GetString(CONFIG_SECTION_NAME, TEXT("RTCBackgroundMode"), RTCBackgroundModeStr, CONFIG_FILE))
+						{
+							LexFromString(PlatformRTCOptions.BackgroundMode, *RTCBackgroundModeStr);					
+						}
+						PlatformOptions.RTCOptions = &PlatformRTCOptions;
+						EosPlatformHandle = EOSPlatformCreate(PlatformOptions);
+
+#if !UE_BUILD_SHIPPING
+						InitializationPath = static_cast<uint8>(EEOSVoiceInitializationPath::LegacyInlineConfig);
+#endif // !UE_BUILD_SHIPPING
 
-					EOS_Platform_RTCOptions PlatformRTCOptions = {};
-					PlatformRTCOptions.ApiVersion = 3;
-					UE_EOS_CHECK_API_MISMATCH(EOS_PLATFORM_RTCOPTIONS_API_LATEST, 3);
-					PlatformRTCOptions.Reserved = nullptr;
-					PlatformRTCOptions.BackgroundMode = EOS_ERTCBackgroundMode::EOS_RTCBM_KeepRoomsAlive;
-					FString RTCBackgroundModeStr;
-					if(GConfig->GetString(TEXT("EOSVoiceChat"), TEXT("RTCBackgroundMode"), RTCBackgroundModeStr, GEngineIni))
-					{
-						LexFromString(PlatformRTCOptions.BackgroundMode, *RTCBackgroundModeStr);					
 					}
-					PlatformOptions.RTCOptions = &PlatformRTCOptions;
 
-					EosPlatformHandle = EOSPlatformCreate(PlatformOptions);
 					if (!EosPlatformHandle)
 					{
 						UE_LOG(LogEOSVoiceChat, Warning, TEXT("FEOSVoiceChat::Initialize CreatePlatform failed"));
@@ -821,7 +888,7 @@ void FEOSVoiceChat::OnAudioDevicesChanged()
 
 void FEOSVoiceChat::OnConfigSectionsChanged(const FString& IniFilename, const TSet<FString>& SectionNames)
 {
-	if (IniFilename == GEngineIni)
+	if (IniFilename == CONFIG_FILE)
 	{
 		LoadConfig();
 	}
@@ -829,27 +896,48 @@ void FEOSVoiceChat::OnConfigSectionsChanged(const FString& IniFilename, const TS
 
 void FEOSVoiceChat::LoadConfig()
 {
-	auto SetFromConfig = [this](const TCHAR* InSettingName)
+	// Helper lambda to apply a single RTC setting
+	auto ApplyRTCSetting = [this](const TCHAR* SettingName, const TCHAR* SettingValue)
+	{
+		const FTCHARToUTF8 Utf8Key(SettingName);
+		const FTCHARToUTF8 Utf8Value(SettingValue);
+
+		EOS_RTC_SetSettingOptions SetSettingOptions = { };
+		SetSettingOptions.ApiVersion = EOS_RTC_SETSETTING_API_LATEST;
+		static_assert(EOS_RTC_SETSETTING_API_LATEST == 1, "EOS_RTC_SetSettingOptions updated, check new fields");
+		SetSettingOptions.SettingName = Utf8Key.Get();
+		SetSettingOptions.SettingValue = Utf8Value.Get();
+		const EOS_EResult Result = EOS_RTC_SetSetting(GetRtcInterface(), &SetSettingOptions);
+		UE_CLOG(Result != EOS_EResult::EOS_Success, LogEOSVoiceChat, Warning,
+			TEXT("SetSetting Key=%s Value=%s Result=%s"), SettingName, SettingValue, *LexToString(Result));
+	};
+
+	// Legacy config load: RTCSetting_<SettingName>=<Value>
+	auto SetFromConfig = [this, &ApplyRTCSetting](const TCHAR* InSettingName)
 	{
-		FString ConfigKey = FString(TEXT("RTCSetting_")) + InSettingName;
+		const FString ConfigKey(FString::Printf(TEXT("RTCSetting_%s"), InSettingName));
 		FString ConfigValue;
-		if (GConfig->GetString(TEXT("EOSVoiceChat"), *ConfigKey, ConfigValue, GEngineIni))
+		if (GConfig->GetString(CONFIG_SECTION_NAME, *ConfigKey, ConfigValue, CONFIG_FILE))
 		{
-			const FTCHARToUTF8 Utf8Key(InSettingName);
-			const FTCHARToUTF8 Utf8Value(*ConfigValue);
-
-			EOS_RTC_SetSettingOptions SetSettingOptions = {};
-			SetSettingOptions.ApiVersion = EOS_RTC_SETSETTING_API_LATEST;
-			static_assert(EOS_RTC_SETSETTING_API_LATEST == 1, "EOS_RTC_SetSettingOptions updated, check new fields");
-			SetSettingOptions.SettingName = Utf8Key.Get();
-			SetSettingOptions.SettingValue = Utf8Value.Get();
-			const EOS_EResult Result = EOS_RTC_SetSetting(GetRtcInterface(), &SetSettingOptions);
-			UE_CLOG(Result != EOS_EResult::EOS_Success, LogEOSVoiceChat, Warning, TEXT("SetSetting Key=%s Value=%s Result=%s"),
-				InSettingName, *ConfigValue, *LexToString(Result));
+			ApplyRTCSetting(InSettingName, *ConfigValue);
 		}
 	};
-
 	SetFromConfig(TEXT("MaxTickInterval"));
+
+	// New config load: Get RTCSettings array with Key=Value entries to apply
+	TArray<FString> RTCSettings;
+	GConfig->GetArray(CONFIG_SECTION_NAME, TEXT("RTCSettings"), RTCSettings, CONFIG_FILE);
+	for (const FString& RTCSetting : RTCSettings)
+	{
+		TArray<FString> Split;
+		RTCSetting.ParseIntoArray(Split, TEXT("="));
+		UE_CLOG(Split.Num() != 2, LogEOSVoiceChat, Warning, TEXT("LoadConfig Cannot parse [%s]"), *RTCSetting);
+
+		if (Split.Num() == 2)
+		{
+			ApplyRTCSetting(*Split[0], *Split[1]);
+		}
+	}
 }
 
 void FEOSVoiceChat::ScheduleReleaseUser(IVoiceChatUser* User)
@@ -894,7 +982,7 @@ bool FEOSVoiceChat::Exec(UWorld* InWorld, const TCHAR* Cmd, FOutputDevice& Ar)
 #define EOS_EXEC_LOG(Fmt, ...) 
 #endif
 
-	if (FParse::Command(&Cmd, TEXT("EOSVOICECHAT")))
+	if (FParse::Command(&Cmd, CONFIG_SECTION_NAME))
 	{
 		const TCHAR* SubCmd = Cmd;
 		if (FParse::Command(&Cmd, TEXT("LIST")))
@@ -918,6 +1006,15 @@ bool FEOSVoiceChat::Exec(UWorld* InWorld, const TCHAR* Cmd, FOutputDevice& Ar)
 			if (FParse::Command(&Cmd, TEXT("INFO")))
 			{
 				EOS_EXEC_LOG(TEXT("Initialized: %s"), *LexToString(IsInitialized()));
+
+#if !UE_BUILD_SHIPPING
+				EOS_EXEC_LOG(TEXT("Initialization Path: %s"), LexToString(static_cast<EEOSVoiceInitializationPath>(InitializationPath)));
+				if (OverridePlatformConfigName.IsSet())
+				{
+					EOS_EXEC_LOG(TEXT("OverridePlatformConfigName: %s"), *OverridePlatformConfigName.GetValue());
+				}
+#endif // !UE_BUILD_SHIPPING
+
 				if (IsInitialized())
 				{
 					EOS_EXEC_LOG(TEXT("Connection Status: %s"), LexToString(ConnectionState));
@@ -934,9 +1031,13 @@ bool FEOSVoiceChat::Exec(UWorld* InWorld, const TCHAR* Cmd, FOutputDevice& Ar)
 	#if !UE_BUILD_SHIPPING
 			else if (FParse::Command(&Cmd, TEXT("INITIALIZE")))
 			{
-				Initialize(FOnVoiceChatInitializeCompleteDelegate::CreateLambda([](const FVoiceChatResult& Result)
+				Initialize(FOnVoiceChatInitializeCompleteDelegate::CreateLambda([this](const FVoiceChatResult& Result)
 				{
-					UE_LOG(LogEOSVoiceChat, Display, TEXT("EOS INITIALIZE success:%s"), *LexToString(Result));
+					UE_LOG(LogEOSVoiceChat, Display, TEXT("EOS INITIALIZE success:%s, Initialization Path: %s"), *LexToString(Result), LexToString(static_cast<EEOSVoiceInitializationPath>(InitializationPath)));
+					if (OverridePlatformConfigName.IsSet())
+					{
+						UE_LOG(LogEOSVoiceChat, Display, TEXT("OverridePlatformConfigName: %s"), *OverridePlatformConfigName.GetValue());
+					}
 				}));
 				return true;
 			}
@@ -1086,5 +1187,7 @@ const TCHAR* LexToString(FEOSVoiceChat::EConnectionState State)
 }
 
 #undef CHECKPIN
+#undef CONFIG_FILE
+#undef CONFIG_SECTION_NAME
 
 #endif // WITH_EOSVOICECHAT
\ No newline at end of file
diff --git a/Engine/Plugins/Online/VoiceChat/EOSVoiceChat/Source/EOSVoiceChat/Private/EOSVoiceChatFactory.cpp b/Engine/Plugins/Online/VoiceChat/EOSVoiceChat/Source/EOSVoiceChat/Private/EOSVoiceChatFactory.cpp
--- a/Engine/Plugins/Online/VoiceChat/EOSVoiceChat/Source/EOSVoiceChat/Private/EOSVoiceChatFactory.cpp
+++ b/Engine/Plugins/Online/VoiceChat/EOSVoiceChat/Source/EOSVoiceChat/Private/EOSVoiceChatFactory.cpp
@@ -15,7 +15,7 @@ IVoiceChatPtr FEOSVoiceChatFactory::CreateInstance()
 	IEOSSDKManager* EOSSDKManager = IEOSSDKManager::Get();
 	check(EOSSDKManager);
 
-	FEOSVoiceChatPtr VoiceChat = MakeShared<FPlatformEOSVoiceChat, ESPMode::ThreadSafe>(*EOSSDKManager, nullptr);
+	FEOSVoiceChatPtr VoiceChat = MakeShared<FPlatformEOSVoiceChat, ESPMode::ThreadSafe>(*EOSSDKManager);
 	Instances.Emplace(VoiceChat);
 	return VoiceChat;
 }
@@ -25,36 +25,22 @@ IVoiceChatPtr FEOSVoiceChatFactory::CreateInstanceWithPlatform(const IEOSPlatfor
 	IEOSSDKManager* EOSSDKManager = IEOSSDKManager::Get();
 	check(EOSSDKManager);
 
-	FEOSVoiceChatPtr VoiceChat = MakeShared<FPlatformEOSVoiceChat, ESPMode::ThreadSafe>(*EOSSDKManager, PlatformHandle);
-	if (VoiceChat->Initialize())
-	{
-		TOptional<bool> bConnectResult;
-		VoiceChat->Connect(FOnVoiceChatConnectCompleteDelegate::CreateLambda([&bConnectResult](const FVoiceChatResult& VoiceChatResult)
-		{
-			bConnectResult = VoiceChatResult.IsSuccess();
-		}));
+	FEOSVoiceChatPtr VoiceChat = MakeShared<FPlatformEOSVoiceChat, ESPMode::ThreadSafe>(*EOSSDKManager);
+	VoiceChat->SetPlatformHandle(PlatformHandle);
+	Instances.Emplace(VoiceChat);
 
-		while (!bConnectResult.IsSet())
-		{
-			PlatformHandle->Tick();
-		}
+	return VoiceChat;
+}
 
-		if (!bConnectResult.GetValue())
-		{
-			UE_LOG(LogEOSVoiceChat, Warning, TEXT("FEOSVoiceChatFactory::CreateInstance Connect failed"));
-			VoiceChat.Reset();
-		}
-	}
-	else
-	{
-		UE_LOG(LogEOSVoiceChat, Warning, TEXT("FEOSVoiceChatFactory::CreateInstance Initialize failed"));
-		VoiceChat.Reset();
-	}
+IVoiceChatPtr FEOSVoiceChatFactory::CreateInstanceWithPlatformConfig(const FString& PlatformConfigName)
+{
+	IEOSSDKManager* EOSSDKManager = IEOSSDKManager::Get();
+	check(EOSSDKManager);
+
+	FEOSVoiceChatPtr VoiceChat = MakeShared<FPlatformEOSVoiceChat, ESPMode::ThreadSafe>(*EOSSDKManager);
+	VoiceChat->SetPlatformConfigName(PlatformConfigName);
+	Instances.Emplace(VoiceChat);
 
-	if(VoiceChat)
-	{
-		Instances.Emplace(VoiceChat);
-	}
 	return VoiceChat;
 }
 
diff --git a/Engine/Plugins/Online/VoiceChat/EOSVoiceChat/Source/EOSVoiceChat/Private/IOS/IOSEOSVoiceChat.cpp b/Engine/Plugins/Online/VoiceChat/EOSVoiceChat/Source/EOSVoiceChat/Private/IOS/IOSEOSVoiceChat.cpp
--- a/Engine/Plugins/Online/VoiceChat/EOSVoiceChat/Source/EOSVoiceChat/Private/IOS/IOSEOSVoiceChat.cpp
+++ b/Engine/Plugins/Online/VoiceChat/EOSVoiceChat/Source/EOSVoiceChat/Private/IOS/IOSEOSVoiceChat.cpp
@@ -10,8 +10,8 @@
 #include "Misc/ConfigCacheIni.h"
 #include "Misc/CoreDelegates.h"
 
-FIOSEOSVoiceChat::FIOSEOSVoiceChat(IEOSSDKManager& InSDKManager, const IEOSPlatformHandlePtr& InPlatformHandle)
-	: FEOSVoiceChat(InSDKManager, InPlatformHandle)
+FIOSEOSVoiceChat::FIOSEOSVoiceChat(IEOSSDKManager& InSDKManager)
+	: FEOSVoiceChat(InSDKManager)
 {
 }
 
diff --git a/Engine/Plugins/Online/VoiceChat/EOSVoiceChat/Source/EOSVoiceChat/Private/IOS/IOSEOSVoiceChat.h b/Engine/Plugins/Online/VoiceChat/EOSVoiceChat/Source/EOSVoiceChat/Private/IOS/IOSEOSVoiceChat.h
--- a/Engine/Plugins/Online/VoiceChat/EOSVoiceChat/Source/EOSVoiceChat/Private/IOS/IOSEOSVoiceChat.h
+++ b/Engine/Plugins/Online/VoiceChat/EOSVoiceChat/Source/EOSVoiceChat/Private/IOS/IOSEOSVoiceChat.h
@@ -9,7 +9,7 @@
 class FIOSEOSVoiceChat : public FEOSVoiceChat
 {
 public:
-	FIOSEOSVoiceChat(IEOSSDKManager& InSDKManager, const IEOSPlatformHandlePtr& InPlatformHandle);
+	FIOSEOSVoiceChat(IEOSSDKManager& InSDKManager);
 	virtual ~FIOSEOSVoiceChat() = default;
 
 	// ~Begin IVoiceChat Interface
diff --git a/Engine/Plugins/Online/VoiceChat/EOSVoiceChat/Source/EOSVoiceChat/Private/Windows/WindowsEOSVoiceChat.cpp b/Engine/Plugins/Online/VoiceChat/EOSVoiceChat/Source/EOSVoiceChat/Private/Windows/WindowsEOSVoiceChat.cpp
--- a/Engine/Plugins/Online/VoiceChat/EOSVoiceChat/Source/EOSVoiceChat/Private/Windows/WindowsEOSVoiceChat.cpp
+++ b/Engine/Plugins/Online/VoiceChat/EOSVoiceChat/Source/EOSVoiceChat/Private/Windows/WindowsEOSVoiceChat.cpp
@@ -6,11 +6,10 @@
 
 #include "HAL/FileManager.h"
 #include "Misc/Paths.h"
-
 #include "Windows/eos_Windows.h"
 
-FWindowsEOSVoiceChat::FWindowsEOSVoiceChat(IEOSSDKManager& InSDKManager, const IEOSPlatformHandlePtr& InPlatformHandle)
-	: FEOSVoiceChat(InSDKManager, InPlatformHandle)
+FWindowsEOSVoiceChat::FWindowsEOSVoiceChat(IEOSSDKManager& InSDKManager)
+	: FEOSVoiceChat(InSDKManager)
 {
 }
 
diff --git a/Engine/Plugins/Online/VoiceChat/EOSVoiceChat/Source/EOSVoiceChat/Private/Windows/WindowsEOSVoiceChat.h b/Engine/Plugins/Online/VoiceChat/EOSVoiceChat/Source/EOSVoiceChat/Private/Windows/WindowsEOSVoiceChat.h
--- a/Engine/Plugins/Online/VoiceChat/EOSVoiceChat/Source/EOSVoiceChat/Private/Windows/WindowsEOSVoiceChat.h
+++ b/Engine/Plugins/Online/VoiceChat/EOSVoiceChat/Source/EOSVoiceChat/Private/Windows/WindowsEOSVoiceChat.h
@@ -9,7 +9,7 @@
 class FWindowsEOSVoiceChat : public FEOSVoiceChat
 {
 public:
-	FWindowsEOSVoiceChat(IEOSSDKManager& InSDKManager, const IEOSPlatformHandlePtr& InPlatformHandle);
+	FWindowsEOSVoiceChat(IEOSSDKManager& InSDKManager);
 	virtual ~FWindowsEOSVoiceChat() = default;
 
 protected:
diff --git a/Engine/Plugins/Online/VoiceChat/EOSVoiceChat/Source/EOSVoiceChat/Public/EOSVoiceChat.h b/Engine/Plugins/Online/VoiceChat/EOSVoiceChat/Source/EOSVoiceChat/Public/EOSVoiceChat.h
--- a/Engine/Plugins/Online/VoiceChat/EOSVoiceChat/Source/EOSVoiceChat/Public/EOSVoiceChat.h
+++ b/Engine/Plugins/Online/VoiceChat/EOSVoiceChat/Source/EOSVoiceChat/Public/EOSVoiceChat.h
@@ -50,6 +50,9 @@ typedef TSharedRef<class FEOSVoiceChatUser, ESPMode::ThreadSafe> FEOSVoiceChatUs
 class FEOSVoiceChat : public TSharedFromThis<FEOSVoiceChat, ESPMode::ThreadSafe>, public IVoiceChat
 {
 public:
+	FEOSVoiceChat(IEOSSDKManager& InSDKManager);
+
+	UE_DEPRECATED(5.8, "Use FEOSVoiceChat(IEOSSDKManager&) and call SetPlatformHandle() instead")
 	FEOSVoiceChat(IEOSSDKManager& InSDKManager, const IEOSPlatformHandlePtr& PlatformHandle);
 
 	virtual ~FEOSVoiceChat();
@@ -71,6 +74,13 @@ public:
 	virtual void ReleaseUser(IVoiceChatUser* VoiceChatUser) override;
 	// ~End IVoiceChat Interface
 
+	// Set the platform config name to use during Initialize(). Only used if no platform handle is set.
+	void SetPlatformConfigName(const FString& PlatformConfigName) { OverridePlatformConfigName = PlatformConfigName; }
+
+	// Set an existing platform handle to use during Initialize(). If set, Initialize() will use this
+	// handle instead of creating one via named config or inline credentials
+	void SetPlatformHandle(const IEOSPlatformHandlePtr& PlatformHandle) { EosPlatformHandle = PlatformHandle; }
+
 	// ~Begin IVoiceChatUser Interface
 	virtual void SetSetting(const FString& Name, const FString& Value) override;
 	virtual FString GetSetting(const FString& Name) override;
@@ -222,6 +232,12 @@ protected:
 	IEOSSDKManager& SDKManager;
 	IEOSPlatformHandlePtr EosPlatformHandle = nullptr;
 	virtual IEOSPlatformHandlePtr EOSPlatformCreate(EOS_Platform_Options& PlatformOptions);
+	// Override for platform config name, set via SetPlatformConfigName(). If set, Initialize() uses this instead of reading from config file
+	TOptional<FString> OverridePlatformConfigName;
+
+#if !UE_BUILD_SHIPPING
+	uint8 InitializationPath = 0;
+#endif // !UE_BUILD_SHIPPING
 
 	static int64 StaticInstanceIdCount;
 	int64 InstanceId = StaticInstanceIdCount++;
diff --git a/Engine/Plugins/Online/VoiceChat/EOSVoiceChat/Source/EOSVoiceChat/Public/EOSVoiceChatFactory.h b/Engine/Plugins/Online/VoiceChat/EOSVoiceChat/Source/EOSVoiceChat/Public/EOSVoiceChatFactory.h
--- a/Engine/Plugins/Online/VoiceChat/EOSVoiceChat/Source/EOSVoiceChat/Public/EOSVoiceChatFactory.h
+++ b/Engine/Plugins/Online/VoiceChat/EOSVoiceChat/Source/EOSVoiceChat/Public/EOSVoiceChatFactory.h
@@ -38,7 +38,8 @@ public:
 	EOSVOICECHAT_API IVoiceChatPtr CreateInstance();
 	/** Create an instance sharing an existing EOS platform. Used to enable interaction with e.g. lobby voice channels via an IVoiceChat interface. */
 	EOSVOICECHAT_API IVoiceChatPtr CreateInstanceWithPlatform(const IEOSPlatformHandlePtr& PlatformHandle);
-
+	/** Create an instance using a specific named platform config. Used to enable interaction with e.g. lobby voice channels via an IVoiceChat interface. */
+	EOSVOICECHAT_API IVoiceChatPtr CreateInstanceWithPlatformConfig(const FString& PlatformConfigName);
 	// ~Begin FSelfRegisteringExec
 #if UE_ALLOW_EXEC_COMMANDS
 	EOSVOICECHAT_API virtual bool Exec(UWorld* InWorld, const TCHAR* Cmd, FOutputDevice& Ar) override;
